Sequencer: MVVM2 branch and Layer Bars

diff --git a/Engine/Source/Editor/SequencerCore/Private/MVVM/Views/SOutlinerView.cpp b/Engine/Source/Editor/SequencerCore/Private/MVVM/Views/SOutlinerView.cpp
new file mode 100644
--- /dev/null
+++ b/Engine/Source/Editor/SequencerCore/Private/MVVM/Views/SOutlinerView.cpp
@@ -0,0 +1,1185 @@
+// Copyright Epic Games, Inc. All Rights Reserved.
+
+#include "MVVM/Views/SOutlinerView.h"
+#include "MVVM/Extensions/IDraggableOutlinerExtension.h"
+#include "MVVM/Extensions/IDroppableExtension.h"
+#include "MVVM/Extensions/IGeometryExtension.h"
+#include "MVVM/Extensions/IPinnableExtension.h"
+#include "MVVM/Extensions/ISelectableExtension.h"
+#include "MVVM/Extensions/ITrackAreaExtension.h"
+#include "MVVM/ViewModels/OutlinerSpacer.h"
+#include "MVVM/ViewModels/OutlinerViewModel.h"
+#include "MVVM/ViewModels/OutlinerViewModelDragDropOp.h"
+#include "MVVM/ViewModels/ViewModelIterators.h"
+#include "MVVM/Views/IOutlinerSelectionHandler.h"
+#include "MVVM/Views/SOutlinerItemViewBase.h"
+#include "MVVM/Views/STrackAreaView.h"
+#include "MVVM/Views/STrackLane.h"
+
+#include "EditorStyleSet.h"
+#include "Algo/BinarySearch.h"
+#include "Algo/Copy.h"
+#include "Framework/Commands/GenericCommands.h"
+#include "Framework/MultiBox/MultiBoxBuilder.h"
+#include "Framework/Application/SlateApplication.h"
+#include "ScopedTransaction.h"
+
+namespace UE
+{
+namespace Sequencer
+{
+
+const FName SOutlinerView::TrackNameColumn("TrackArea");
+
+void SOutlinerViewRow::Construct(const FArguments& InArgs, const TSharedRef<STableViewBase>& OwnerTableView, TWeakViewModelPtr<IOutlinerExtension> InWeakModel)
+{
+	WeakModel = InWeakModel;
+
+	OnDetectDrag = InArgs._OnDetectDrag;
+	OnGenerateWidgetForColumn = InArgs._OnGenerateWidgetForColumn;
+
+	SMultiColumnTableRow::Construct(
+		SMultiColumnTableRow::FArguments()
+			.OnDragDetected(this, &SOutlinerViewRow::OnDragDetected)
+			.OnCanAcceptDrop(this, &SOutlinerViewRow::OnCanAcceptDrop)
+			.OnAcceptDrop(this, &SOutlinerViewRow::OnAcceptDrop)
+			.ShowSelection(IsSelectable())
+			.Padding(FMargin(0.f)),
+		OwnerTableView);
+}
+
+SOutlinerViewRow::~SOutlinerViewRow()
+{
+}
+
+TSharedRef<SWidget> SOutlinerViewRow::GenerateWidgetForColumn(const FName& ColumnId)
+{
+	TViewModelPtr<IOutlinerExtension> DataModel = WeakModel.Pin();
+	return DataModel ? OnGenerateWidgetForColumn.Execute(DataModel, ColumnId, SharedThis(this)) : SNullWidget::NullWidget;
+}
+
+void SOutlinerViewRow::Tick(const FGeometry& AllottedGeometry, const double InCurrentTime, const float InDeltaTime)
+{
+	StaticCastSharedPtr<SOutlinerView>(OwnerTablePtr.Pin())->ReportChildRowGeometry(WeakModel.Pin(), AllottedGeometry);
+}
+
+bool SOutlinerViewRow::IsSelectable() const
+{
+	TSharedPtr<ISelectableExtension> Selectable = WeakModel.ImplicitPin();
+	return !Selectable || Selectable->IsSelectable() != ESelectionIntent::Never;
+}
+
+FReply SOutlinerViewRow::OnDragDetected( const FGeometry& InGeometry, const FPointerEvent& InPointerEvent )
+{
+	if (OnDetectDrag.IsBound())
+	{
+		return OnDetectDrag.Execute(InGeometry, InPointerEvent, SharedThis(this));
+	}
+	return FReply::Unhandled();
+}
+
+int32 SOutlinerViewRow::OnPaintDropIndicator(EItemDropZone InItemDropZone, const FPaintArgs& Args, const FGeometry& AllottedGeometry, const FSlateRect& MyCullingRect, FSlateWindowElementList& OutDrawElements, int32 LayerId, const FWidgetStyle& InWidgetStyle, bool bParentEnabled) const
+{
+	// Draw feedback for user dropping an item above, below, or onto a row.
+	const FSlateBrush* DropIndicatorBrush = GetDropIndicatorBrush(InItemDropZone);
+
+	// Offset by the indentation amount
+	static float OffsetX = 10.0f;
+	FVector2D Offset(OffsetX * GetIndentLevel(), 0.f);
+	FSlateDrawElement::MakeBox
+	(
+		OutDrawElements,
+		LayerId++,
+		AllottedGeometry.ToPaintGeometry(FVector2D(AllottedGeometry.GetLocalSize() - Offset), FSlateLayoutTransform(Offset)),
+		DropIndicatorBrush,
+		ESlateDrawEffect::None,
+		DropIndicatorBrush->GetTint(InWidgetStyle) * InWidgetStyle.GetColorAndOpacityTint()
+	);
+
+	return LayerId;
+}
+
+TOptional<EItemDropZone> SOutlinerViewRow::OnCanAcceptDrop(const FDragDropEvent& DragDropEvent, EItemDropZone InItemDropZone, TWeakViewModelPtr<IOutlinerExtension> InDataModel)
+{
+	TViewModelPtr<IOutlinerExtension> ThisModel = InDataModel.Pin();
+	if (!ThisModel)
+	{
+		return TOptional<EItemDropZone>();
+	}
+	else if (InItemDropZone == EItemDropZone::BelowItem && ThisModel->IsExpanded())
+	{
+		// Cannot drop immediately below items that are expanded since
+		// it looks like you are dropping into the item, but the object
+		// will end up below all its expanded children.
+		return TOptional<EItemDropZone>();
+	}
+
+	// The model we are indirectly dropping into (either this model, or its parent)
+	FViewModelPtr DropModel = ThisModel;
+	// The model we are directly interacting with - only set when attaching before/after
+	FViewModelPtr TargetModel;
+
+	if (!DropModel)
+	{
+		return TOptional<EItemDropZone>();
+	}
+
+	// When dropping above or below an item, we always forward to its parent
+	if (InItemDropZone == EItemDropZone::AboveItem || InItemDropZone == EItemDropZone::BelowItem)
+	{
+		TargetModel = DropModel;
+		DropModel = DropModel->CastParent<IOutlinerDropTargetOutlinerExtension>();
+	}
+
+	TViewModelPtr<IOutlinerDropTargetOutlinerExtension> DropTarget = DropModel.ImplicitCast();
+	if (!DropTarget)
+	{
+		return TOptional<EItemDropZone>();
+	}
+
+	return DropTarget->CanAcceptDrop(TargetModel, DragDropEvent, InItemDropZone);
+}
+
+FReply SOutlinerViewRow::OnAcceptDrop(const FDragDropEvent& DragDropEvent, EItemDropZone InItemDropZone, TWeakViewModelPtr<IOutlinerExtension> InDataModel)
+{
+	TViewModelPtr<IOutlinerExtension> ThisModel = InDataModel.Pin();
+	if (!ThisModel)
+	{
+		return FReply::Unhandled();
+	}
+	else if (InItemDropZone == EItemDropZone::BelowItem && ThisModel->IsExpanded())
+	{
+		// Cannot drop immediately below items that are expanded since
+		// it looks like you are dropping into the item, but the object
+		// will end up below all its expanded children.
+		return FReply::Unhandled();
+	}
+
+	// The model we are indirectly dropping into (either this model, or its parent)
+	FViewModelPtr DropModel = ThisModel;
+	// The model we are directly interacting with - only set when attaching before/after
+	FViewModelPtr TargetModel;
+
+	// When dropping above or below an item, we always forward to its parent
+	if (InItemDropZone == EItemDropZone::AboveItem || InItemDropZone == EItemDropZone::BelowItem)
+	{
+		TargetModel = DropModel;
+		DropModel = DropModel->CastParent<IOutlinerDropTargetOutlinerExtension>();
+	}
+
+	TViewModelPtr<IOutlinerDropTargetOutlinerExtension> DropTarget = DropModel.ImplicitCast();
+	if (DropTarget)
+	{
+		DropTarget->PerformDrop(TargetModel, DragDropEvent, InItemDropZone);
+		return FReply::Handled();
+	}
+
+	return FReply::Unhandled();
+}
+
+TViewModelPtr<IOutlinerExtension> SOutlinerViewRow::GetDataModel() const
+{
+	return WeakModel.Pin();
+}
+
+TSharedPtr<STrackLane> SOutlinerViewRow::GetTrackLane(bool bOnlyOwnTrackLane) const
+{
+	if (!bOnlyOwnTrackLane)
+	{
+		// Return the track lane, regardless of it being created by our own outliner item, or
+		// being a reference to a parent outliner item's track lane.
+		return TrackLane;
+	}
+	if (TrackLane && TrackLane->GetOutlinerItem() == GetDataModel())
+	{
+		// Return the track lane only if it was created by our own outliner item.
+		return TrackLane;
+	}
+	return nullptr;
+}
+
+void SOutlinerViewRow::SetTrackLane(const TSharedPtr<STrackLane>& InTrackLane)
+{
+	TrackLane = InTrackLane;
+}
+
+const FSlateBrush* SOutlinerViewRow::GetBorder() const 
+{
+	TSharedPtr<IOutlinerExtension> OutlinerItem = WeakModel.ImplicitPin();
+	if (OutlinerItem && !OutlinerItem->HasBackground())
+	{
+		return nullptr;
+	}
+
+	return SMultiColumnTableRow::GetBorder();
+}
+
+void SOutlinerView::Construct(const FArguments& InArgs, TWeakPtr<FOutlinerViewModel> InWeakOutliner, const TSharedRef<STrackAreaView>& InTrackArea)
+{
+	WeakOutliner = InWeakOutliner;
+	TrackArea = InTrackArea;
+	bUpdatingTreeSelection = false;
+	bRightMouseButtonDown = false;
+	bShowPinnedNodes = false;
+	bSelectionChangesPending = false;
+	bRefreshPhysicalGeometry = false;
+	bForwardToSelectionHandler = true;
+
+	SelectionHandlerAttribute = InArgs._SelectionHandler;
+
+	HeaderRow = SNew(SHeaderRow).Visibility(EVisibility::Collapsed);
+
+	SetupColumns(InArgs);
+
+	WeakOutliner.Pin()->OnRefreshed.AddSP(this, &SOutlinerView::Refresh);
+
+	STreeView::Construct
+	(
+		STreeView::FArguments()
+		.TreeItemsSource(&RootNodes)
+		.OnItemToString_Debug(this, &SOutlinerView::OnItemToString_Debug)
+		.SelectionMode(ESelectionMode::Multi)
+		.OnGenerateRow(this, &SOutlinerView::OnGenerateRow)
+		.OnGetChildren(this, &SOutlinerView::OnGetChildren)
+		.HeaderRow(HeaderRow)
+		.ExternalScrollbar(InArgs._ExternalScrollbar)
+		.OnExpansionChanged(this, &SOutlinerView::OnExpansionChanged)
+		.OnTreeViewScrolled(this, &SOutlinerView::HandleTableViewScrolled)
+		.AllowOverscroll(EAllowOverscroll::No)
+		.OnContextMenuOpening( this, &SOutlinerView::OnContextMenuOpening )
+		.OnSetExpansionRecursive(this, &SOutlinerView::SetItemExpansionRecursive)
+		.HighlightParentNodesForSelection(true)
+		.AllowInvisibleItemSelection(true)  //without this we deselect everything when we filter or we collapse, etc..
+	);
+}
+
+SOutlinerView::~SOutlinerView()
+{
+	TrackArea.Reset();
+	SlaveTreeViews.Empty();
+}
+
+void SOutlinerView::Tick(const FGeometry& AllottedGeometry, const double InCurrentTime, const float InDeltaTime)
+{
+	const bool bWasPendingRefresh = IsPendingRefresh();
+
+	STreeView::Tick(AllottedGeometry, InCurrentTime, InDeltaTime);
+
+	if (bWasPendingRefresh || bRefreshPhysicalGeometry)
+	{
+		UpdatePhysicalGeometry(true);
+		bRefreshPhysicalGeometry = false;
+	}
+}
+
+TSharedPtr<FOutlinerViewModel> SOutlinerView::GetOutlinerModel() const
+{
+	return WeakOutliner.Pin();
+}
+
+FString SOutlinerView::OnItemToString_Debug(TWeakViewModelPtr<IOutlinerExtension> InWeakModel)
+{
+	TViewModelPtr<IOutlinerExtension> Model = InWeakModel.ImplicitPin();
+	return Model ? Model->GetLabel().ToString() : FString();
+}
+
+int32 SOutlinerView::OnPaint(const FPaintArgs& Args, const FGeometry& AllottedGeometry, const FSlateRect& MyCullingRect, FSlateWindowElementList& OutDrawElements, int32 LayerId, const FWidgetStyle& InWidgetStyle, bool bParentEnabled) const
+{
+	return STreeView::OnPaint(Args, AllottedGeometry, MyCullingRect, OutDrawElements, LayerId, InWidgetStyle, bParentEnabled);
+}
+
+FReply SOutlinerView::OnDragRow(const FGeometry&, const FPointerEvent&, TSharedRef<SOutlinerViewRow> InRow)
+{
+	TSharedPtr<FOutlinerViewModel> Outliner = WeakOutliner.Pin();
+	if (!Outliner)
+	{
+		return FReply::Unhandled();
+	}
+
+	TArray<TWeakViewModelPtr<IOutlinerExtension>> Selection = GetSelectedItems();
+	if (Selection.Num() > 0)
+	{
+		for (int32 Index = Selection.Num()-1; Index >= 0; --Index)
+		{
+			TSharedPtr<IDraggableOutlinerExtension> Draggable = Selection[Index].ImplicitPin();
+			if (!Draggable || !Draggable->CanDrag())
+			{
+				// Order is not important so we can opt for performance with RemoveAtSwap
+				Selection.RemoveAtSwap(Index, 1, false);
+			}
+		}
+
+		// If there were no nodes selected, don't start a drag drop operation.
+		if (Selection.Num() == 0)
+		{
+			return FReply::Unhandled();
+		}
+
+		TSharedRef<FDragDropOperation> DragDropOp = Outliner->InitiateDrag( MoveTemp(Selection) );
+		return FReply::Handled().BeginDragDrop( DragDropOp );
+	}
+
+	return FReply::Unhandled();
+}
+
+void SOutlinerView::ReportChildRowGeometry(const TViewModelPtr<IOutlinerExtension>& InNode, const FGeometry& InGeometry)
+{
+	FCachedGeometry* PhysicalNode = PhysicalNodes.FindByPredicate(
+			[InNode](FCachedGeometry& Item) { return Item.WeakItem == InNode; });
+	if (PhysicalNode)
+	{
+		PhysicalNode->PhysicalHeight = InGeometry.Size.Y;
+	}
+}
+
+void SOutlinerView::GetVisibleItems(TArray<TViewModelPtr<IOutlinerExtension>>& OutItems) const
+{
+	OutItems.Reserve(OutItems.Num() + PhysicalNodes.Num());
+	for (const FCachedGeometry& CachedGeometry : PhysicalNodes)
+	{
+		if (TViewModelPtr<IOutlinerExtension> Item = CachedGeometry.WeakItem.Pin())
+		{
+			OutItems.Add(Item);
+		}
+	}
+}
+
+void SOutlinerView::ForceSetSelectedItems(const TSet<TWeakPtr<FViewModel>>& InItems)
+{
+	TGuardValue<bool> Guard(bForwardToSelectionHandler, false);
+
+	Private_ClearSelection();
+	for (const TWeakPtr<FViewModel>& Item : InItems)
+	{
+		TWeakViewModelPtr<IOutlinerExtension> WeakItem(Item);
+		if (WeakItem.Pin())
+		{
+			Private_SetItemSelection(WeakItem, true, false);
+		}
+	}
+
+	Private_SignalSelectionChanged(ESelectInfo::Direct);
+}
+
+TViewModelPtr<IOutlinerExtension> SOutlinerView::HitTestNode(float InPhysical) const
+{
+	// Find the first node with a top after the specified value - the hit node must be the one preceeding this
+	const int32 FoundIndex = Algo::UpperBoundBy(PhysicalNodes, InPhysical, &FCachedGeometry::PhysicalTop) - 1;
+	if (FoundIndex >= 0)
+	{
+		return PhysicalNodes[FoundIndex].WeakItem.Pin();
+	}
+
+	return nullptr;
+}
+
+float SOutlinerView::PhysicalToVirtual(float InPhysical) const
+{
+	// Find the first node with a top after the specified value - the hit node must be the one preceeding this
+	const int32 FoundIndex = Algo::UpperBoundBy(PhysicalNodes, InPhysical, &FCachedGeometry::PhysicalTop) - 1;
+	if (FoundIndex >= 0)
+	{
+		FCachedGeometry Found = PhysicalNodes[FoundIndex];
+		const float FractionalHeight = (Found.PhysicalHeight != 0.f) ? (InPhysical - Found.PhysicalTop) / Found.PhysicalHeight : 0.f;
+		return Found.VirtualTop + Found.VirtualHeight * FractionalHeight;
+	}
+
+	if (PhysicalNodes.Num())
+	{
+		const FCachedGeometry& First = PhysicalNodes[0];
+		if (InPhysical < First.PhysicalTop)
+		{
+			return First.VirtualTop + (InPhysical - First.PhysicalTop);
+		}
+		else
+		{
+			const FCachedGeometry& Last = PhysicalNodes.Last();
+			return First.VirtualTop + (InPhysical - Last.PhysicalTop);
+		}
+	}
+
+	return InPhysical;
+}
+
+float SOutlinerView::VirtualToPhysical(float InVirtual) const
+{
+	// Find the first node with a top after the specified value - the hit node must be the one preceeding this
+	int32 FoundIndex = Algo::UpperBoundBy(PhysicalNodes, InVirtual, &FCachedGeometry::VirtualTop) - 1;
+
+	if (FoundIndex >= 0)
+	{
+		const FCachedGeometry& Found = PhysicalNodes[FoundIndex];
+		const float FractionalHeight = (Found.VirtualHeight != 0.f) ? (InVirtual - Found.VirtualTop) / Found.VirtualHeight : 0.f;
+		return Found.PhysicalTop + Found.PhysicalHeight * FractionalHeight;
+	}
+
+	if (PhysicalNodes.Num())
+	{
+		const FCachedGeometry& Last = PhysicalNodes.Last();
+		return Last.PhysicalTop + (InVirtual - Last.VirtualTop);
+	}
+
+	return InVirtual;
+}
+
+void SOutlinerView::SetupColumns(const FArguments& InArgs)
+{
+	TSharedPtr<FEditorViewModel> EditorViewModel = WeakOutliner.Pin()->GetEditor();
+
+	// Define a column for the Outliner
+	auto GenerateOutliner = [=](const TWeakViewModelPtr<IOutlinerExtension>& InWeakModel, const TSharedRef<SOutlinerViewRow>& InRow) -> TSharedRef<SWidget>
+	{
+		if (TSharedPtr<IOutlinerExtension> OutlinerItem = InWeakModel.ImplicitPin())
+		{
+			return OutlinerItem->CreateOutlinerView(FCreateOutlinerViewParams{ InRow, EditorViewModel });
+		}
+
+		ensureMsgf(false, TEXT("Attempting to create an outliner widget for a view model that is either dead, or not an outliner item."));
+		return SNew(SBox).HeightOverride(10.f);
+	};
+
+	Columns.Add("Outliner", FOutlinerViewColumn(GenerateOutliner, 1.f));
+
+	// Now populate the header row with the columns
+	for (TTuple<FName, FOutlinerViewColumn>& Pair : Columns)
+	{
+		if (Pair.Key != TrackNameColumn)
+		{
+			HeaderRow->AddColumn(
+				SHeaderRow::Column(Pair.Key)
+				.FillWidth(Pair.Value.Width)
+			);
+		}
+	}
+}
+
+void SOutlinerView::UpdateTrackArea()
+{
+	// Add or remove the column
+	if (const FOutlinerViewColumn* Column = Columns.Find(TrackNameColumn))
+	{
+		HeaderRow->AddColumn(
+			SHeaderRow::Column(TrackNameColumn)
+			.FillWidth(Column->Width)
+		);
+	}
+}
+
+void SOutlinerView::AddSlaveTreeView(TSharedPtr<SOutlinerView> SlaveTreeView)
+{
+	SlaveTreeViews.Add(SlaveTreeView);
+	SlaveTreeView->SetMasterTreeView(SharedThis(this));
+}
+
+void SOutlinerView::OnRightMouseButtonDown(const FPointerEvent& MouseEvent)
+{
+	STreeView::OnRightMouseButtonDown(MouseEvent);
+	bRightMouseButtonDown = true;
+}
+
+void SOutlinerView::OnRightMouseButtonUp(const FPointerEvent& MouseEvent)
+{
+	STreeView::OnRightMouseButtonUp(MouseEvent);
+	bRightMouseButtonDown = false;
+}
+
+FReply SOutlinerView::OnDragOver(const FGeometry& MyGeometry, const FDragDropEvent& DragDropEvent)
+{
+	FReply Result = STreeView::OnDragOver(MyGeometry, DragDropEvent);
+
+	if (!Result.IsEventHandled())
+	{
+		TSharedPtr<FOutlinerViewModelDragDropOp> DragDropOp = DragDropEvent.GetOperationAs<FOutlinerViewModelDragDropOp>();
+		if (DragDropOp.IsValid())
+		{
+			// Reset tooltip to indicate that the dragged objects can be moved to the root (ie. unparented)
+			DragDropOp->ResetToDefaultToolTip();
+		}
+
+		return Result.Handled();
+	}
+	return Result;
+}
+
+FReply SOutlinerView::OnDrop(const FGeometry& MyGeometry, const FDragDropEvent& DragDropEvent)
+{
+	TSharedPtr<FOutlinerViewModel> Outliner = WeakOutliner.Pin();
+	if (!Outliner)
+	{
+		return FReply::Handled();
+	}
+
+	FReply Result = STreeView::OnDrop(MyGeometry, DragDropEvent);
+
+	TSharedPtr<FOutlinerViewModelDragDropOp> DragDropOp = DragDropEvent.GetOperationAs<FOutlinerViewModelDragDropOp>();
+	if (!Result.IsEventHandled() && DragDropOp.IsValid() && DragDropOp->GetDraggedViewModels().Num())
+	{
+		const FScopedTransaction Transaction(NSLOCTEXT("SequencerTrackNode", "MoveItems", "Move items."));
+
+		TViewModelPtr<IDroppableExtension> Root = Outliner->GetRootItem().ImplicitCast();
+		if (Root)
+		{
+			Root->ProcessDragOperation(*DragDropOp);
+		}
+
+		return FReply::Handled();
+	}
+	
+	return Result;
+}
+
+FReply SOutlinerView::OnKeyDown(const FGeometry& MyGeometry, const FKeyEvent& InKeyEvent)
+{
+	const TArray<TWeakViewModelPtr<IOutlinerExtension>>& ItemsSourceRef = (*this->ItemsSource);
+
+	auto IsSelectable = [](const TWeakViewModelPtr<IOutlinerExtension>& WeakModel)
+	{
+		TSharedPtr<ISelectableExtension> Selectable = WeakModel.ImplicitPin();
+		return !Selectable || EnumHasAnyFlags(Selectable->IsSelectable(), ESelectionIntent::PersistentSelection);
+	};
+
+	// Don't respond to key-presses containing "Alt" as a modifier
+	if (ItemsSourceRef.Num() > 0 && !InKeyEvent.IsAltDown())
+	{
+		bool bWasHandled = false;
+		NullableItemType ItemNavigatedTo(nullptr);
+
+		// Check for selection manipulation keys (Up, Down, Home, End, PageUp, PageDown)
+		if (InKeyEvent.GetKey() == EKeys::Up)
+		{
+			int32 SelectionIndex = 0;
+			if (TListTypeTraits<TWeakViewModelPtr<IOutlinerExtension>>::IsPtrValid(SelectorItem))
+			{
+				SelectionIndex = ItemsSourceRef.Find(TListTypeTraits<TWeakViewModelPtr<IOutlinerExtension>>::NullableItemTypeConvertToItemType(SelectorItem));
+			}
+
+			--SelectionIndex;
+
+			for (; SelectionIndex >=0; --SelectionIndex)
+			{
+				if (IsSelectable(ItemsSourceRef[SelectionIndex]))
+				{
+					ItemNavigatedTo = ItemsSourceRef[SelectionIndex];
+					break;
+				}
+			}
+			bWasHandled = true;
+		}
+		else if (InKeyEvent.GetKey() == EKeys::Down)
+		{
+			int32 SelectionIndex = 0;
+			if (TListTypeTraits<TWeakViewModelPtr<IOutlinerExtension>>::IsPtrValid(SelectorItem))
+			{
+				SelectionIndex = ItemsSourceRef.Find(TListTypeTraits<TWeakViewModelPtr<IOutlinerExtension>>::NullableItemTypeConvertToItemType(SelectorItem));
+			}
+
+			++SelectionIndex;
+
+			for (; SelectionIndex < ItemsSourceRef.Num(); ++SelectionIndex)
+			{
+				if (IsSelectable(ItemsSourceRef[SelectionIndex]))
+				{
+					ItemNavigatedTo = ItemsSourceRef[SelectionIndex];
+					break;
+				}
+			}
+			bWasHandled = true;
+		}
+		else if (InKeyEvent.GetKey() == EKeys::Home)
+		{
+			// Select the first item
+			for (int32 SelectionIndex = 0; SelectionIndex < ItemsSourceRef.Num(); ++SelectionIndex)
+			{
+				if (IsSelectable(ItemsSourceRef[SelectionIndex]))
+				{
+					ItemNavigatedTo = ItemsSourceRef[SelectionIndex];
+					break;
+				}
+			}
+			bWasHandled = true;
+		}
+		else if (InKeyEvent.GetKey() == EKeys::End)
+		{
+			// Select the last item
+			for (int32 SelectionIndex = ItemsSourceRef.Num() -1; SelectionIndex >=0 ; --SelectionIndex)
+			{
+				if (IsSelectable(ItemsSourceRef[SelectionIndex]))
+				{
+					ItemNavigatedTo = ItemsSourceRef[SelectionIndex];
+					break;
+				}
+			}
+			bWasHandled = true;
+		}
+		else if (InKeyEvent.GetKey() == EKeys::PageUp)
+		{
+			int32 SelectionIndex = 0;
+			if (TListTypeTraits<TWeakViewModelPtr<IOutlinerExtension>>::IsPtrValid(SelectorItem))
+			{
+				SelectionIndex = ItemsSourceRef.Find(TListTypeTraits<TWeakViewModelPtr<IOutlinerExtension>>::NullableItemTypeConvertToItemType(SelectorItem));
+			}
+
+			int32 NumItemsInAPage = GetNumLiveWidgets();
+			int32 Remainder = NumItemsInAPage % GetNumItemsPerLine();
+			NumItemsInAPage -= Remainder;
+
+			if (SelectionIndex >= NumItemsInAPage)
+			{
+				// Select an item on the previous page
+				SelectionIndex = SelectionIndex - NumItemsInAPage;
+
+				// Scan up for the first selectable node
+				for (; SelectionIndex >= 0; --SelectionIndex)
+				{
+					if (IsSelectable(ItemsSourceRef[SelectionIndex]))
+					{
+						ItemNavigatedTo = ItemsSourceRef[SelectionIndex];
+						break;
+					}
+				}
+			}
+
+			// If we had less than a page to jump, or we haven't found a selectable node yet,
+			// scan back toward our current node until we find one.
+			if (!ItemNavigatedTo.Pin())
+			{
+				SelectionIndex = 0;
+				for (; SelectionIndex < ItemsSourceRef.Num(); ++SelectionIndex)
+				{
+					if (IsSelectable(ItemsSourceRef[SelectionIndex]))
+					{
+						ItemNavigatedTo = ItemsSourceRef[SelectionIndex];
+						break;
+					}
+				}
+			}
+
+			bWasHandled = true;
+		}
+		else if (InKeyEvent.GetKey() == EKeys::PageDown)
+		{
+			int32 SelectionIndex = 0;
+			if (TListTypeTraits<TWeakViewModelPtr<IOutlinerExtension>>::IsPtrValid(SelectorItem))
+			{
+				SelectionIndex = ItemsSourceRef.Find(TListTypeTraits<TWeakViewModelPtr<IOutlinerExtension>>::NullableItemTypeConvertToItemType(SelectorItem));
+			}
+
+			int32 NumItemsInAPage = GetNumLiveWidgets();
+			int32 Remainder = NumItemsInAPage % GetNumItemsPerLine();
+			NumItemsInAPage -= Remainder;
+
+
+			if (SelectionIndex < ItemsSourceRef.Num() - NumItemsInAPage)
+			{
+				// Select an item on the next page
+				SelectionIndex = SelectionIndex + NumItemsInAPage;
+
+				for (; SelectionIndex < ItemsSourceRef.Num(); ++SelectionIndex)
+				{
+					if (IsSelectable(ItemsSourceRef[SelectionIndex]))
+					{
+						ItemNavigatedTo = ItemsSourceRef[SelectionIndex];
+						break;
+					}
+				}
+			}
+
+			// If we had less than a page to jump, or we haven't found a selectable node yet,
+			// scan back toward our current node until we find one.
+			if (!ItemNavigatedTo.Pin())
+			{
+				SelectionIndex = ItemsSourceRef.Num() - 1;
+				for (; SelectionIndex >= 0; --SelectionIndex)
+				{
+					if (IsSelectable(ItemsSourceRef[SelectionIndex]))
+					{
+						ItemNavigatedTo = ItemsSourceRef[SelectionIndex];
+						break;
+					}
+				}
+			}
+			bWasHandled = true;
+		}
+		else if (InKeyEvent.GetKey() == EKeys::SpaceBar)
+		{
+			// SListView behavior is: Change selected status of item. Bypass that here. We don't want that, but instead want spacebar to go to SequencerCommands (toggle play)
+			return FReply::Unhandled();
+		}
+
+		if (TListTypeTraits<TWeakViewModelPtr<IOutlinerExtension>>::IsPtrValid(ItemNavigatedTo))
+		{
+			TWeakViewModelPtr<IOutlinerExtension> ItemToSelect(TListTypeTraits<TWeakViewModelPtr<IOutlinerExtension>>::NullableItemTypeConvertToItemType(ItemNavigatedTo));
+			NavigationSelect(ItemToSelect, InKeyEvent);
+		}
+
+		if (bWasHandled)
+		{
+			return FReply::Handled();
+		}
+	}
+
+	return STreeView<TWeakViewModelPtr<IOutlinerExtension>>::OnKeyDown(MyGeometry, InKeyEvent);
+}
+
+void SOutlinerView::Private_SetItemSelection( TWeakViewModelPtr<IOutlinerExtension> TheItem, bool bShouldBeSelected, bool bWasUserDirected )
+{
+	if (TSharedPtr<FOutlinerSpacer> Spacer = TheItem.ImplicitPin())
+	{
+		return;
+	}
+
+	STreeView::Private_SetItemSelection( TheItem, bShouldBeSelected, bWasUserDirected );
+}
+
+void SOutlinerView::Private_ClearSelection()
+{
+	STreeView::Private_ClearSelection();
+}
+
+void SOutlinerView::Private_SelectRangeFromCurrentTo( TWeakViewModelPtr<IOutlinerExtension> InRangeSelectionEnd )
+{
+	STreeView::Private_SelectRangeFromCurrentTo(InRangeSelectionEnd);
+}
+
+void SOutlinerView::Private_SignalSelectionChanged(ESelectInfo::Type SelectInfo)
+{
+	// If bForwardToSelectionHandler is set here that means we are responding to a
+	// selection change event that originated from a user interaction with this outliner view.
+	// Therefore we must go through and re-assign selection states to the selection handler.
+
+	STreeView::Private_SignalSelectionChanged(SelectInfo);
+
+	TSharedPtr<IOutlinerSelectionHandler> SelectionHandler = SelectionHandlerAttribute.Get();
+	if (SelectionHandler && bForwardToSelectionHandler)
+	{
+		TGuardValue<bool> Guard(bForwardToSelectionHandler, false);
+		SelectionHandler->SelectOutlinerItems(GetSelectedItems(), bRightMouseButtonDown);
+	}
+}
+
+TSharedPtr<SWidget> SOutlinerView::OnContextMenuOpening()
+{
+	TSharedPtr<SWidget> ContextMenu;
+
+	// Open a context menu for the first selected item if it is selectable
+	for (TWeakViewModelPtr<IOutlinerExtension> WeakSelectedItem : GetSelectedItems())
+	{
+		TViewModelPtr<IOutlinerExtension> SelectedItem = WeakSelectedItem.Pin();
+
+		// Items are selectable by default, but can opt-in to conditionally selectable
+		// TODO: This should be checked in the selection handler?
+		TViewModelPtr<ISelectableExtension> Selectable = WeakSelectedItem.ImplicitPin();
+		const bool bIsSelectable = !Selectable || EnumHasAnyFlags(Selectable->IsSelectable(), ESelectionIntent::ContextMenu);
+
+		if (bIsSelectable)
+		{
+			TSharedPtr<FEditorViewModel> EditorViewModel = GetOutlinerModel()->GetEditor();
+			ContextMenu = SelectedItem->CreateContextMenuWidget(EditorViewModel);
+		}
+		break;
+	}
+
+	// Otherwise, add a general menu for options
+	TSharedPtr<FOutlinerViewModel> Outliner = WeakOutliner.Pin();
+	if (Outliner && !ContextMenu)
+	{
+		ContextMenu = Outliner->CreateContextMenuWidget();
+	}
+
+	return ContextMenu;
+}
+
+void SOutlinerView::Refresh()
+{
+	TSharedPtr<FOutlinerViewModel> Outliner = WeakOutliner.Pin();
+
+	RootNodes.Reset();
+	if (!Outliner)
+	{
+		return;
+	}
+
+	for (TWeakViewModelPtr<IOutlinerExtension> WeakItem : Outliner->GetTopLevelItems())
+	{
+		TViewModelPtr<IOutlinerExtension> Extension = WeakItem.Pin();
+		if (!Extension || Extension->IsFilteredOut())
+		{
+			continue;
+		}
+
+		// Only add pinned nodes if this is showing pinned only
+		bool bIsPinned = false;
+
+		constexpr bool bIncludeThis = true;
+		for (const TViewModelPtr<IPinnableExtension>& Pinnable : Extension.AsModel()->GetAncestorsOfType<IPinnableExtension>(bIncludeThis))
+		{
+			if (Pinnable->IsPinned())
+			{
+				bIsPinned = true;
+				break;
+			}
+		}
+
+		if (bIsPinned == bShowPinnedNodes)
+		{
+			RootNodes.Add(Extension.AsModel());
+		}
+	}
+
+	// Reset item expansion since we don't know if any expansion states may have changed in-between refreshes
+	{
+		STreeView::OnExpansionChanged.Unbind();
+
+		ClearExpandedItems();
+
+		FViewModelPtr RootItem = Outliner->GetRootItem();
+		if (RootItem)
+		{
+			for (TViewModelPtr<IOutlinerExtension> Child : RootItem->GetDescendantsOfType<IOutlinerExtension>())
+			{
+				SetItemExpansion(Child, Child->IsExpanded());
+			}
+		}
+
+		STreeView::OnExpansionChanged.BindSP(this, &SOutlinerView::OnExpansionChanged);
+	}
+
+	RebuildList();
+	//RequestTreeRefresh();
+
+	for (TSharedPtr<SOutlinerView> SlaveTreeView : SlaveTreeViews)
+	{
+		SlaveTreeView->Refresh();
+	}
+}
+
+void SOutlinerView::ScrollByDelta(float DeltaInSlateUnits)
+{
+	ScrollBy( GetCachedGeometry(), DeltaInSlateUnits, EAllowOverscroll::No );
+}
+
+float SOutlinerView::GetVirtualTop() const
+{
+	return VirtualTop;
+}
+
+void SOutlinerView::HandleTableViewScrolled(double InScrollOffset)
+{
+	if (IsPendingRefresh())
+	{
+		bRefreshPhysicalGeometry = true;
+	}
+	else
+	{
+		UpdatePhysicalGeometry(false);
+	}
+}
+
+void SOutlinerView::UpdatePhysicalGeometry(bool bIsRefresh)
+{
+	if (bIsRefresh)
+	{
+		// We need to first update virtual geometry in case the model hierarchy has changed, which is
+		// very likely if we end up here from a tree view refresh.
+		if (TSharedPtr<FOutlinerViewModel> Outliner = WeakOutliner.Pin())
+		{
+			if (FViewModelPtr RootItem = Outliner->GetRootItem())
+			{
+				IGeometryExtension::UpdateVirtualGeometry(0.f, RootItem);
+			}
+		}
+	}
+
+	const FChildren* Children = GetConstructedTableItems();
+	const int32 NumSlots = Children->Num();
+
+	PhysicalNodes.Empty();
+
+	if (NumSlots == 0)
+	{
+		return;
+	}
+
+	const float FirstLineScrollOffset = GetFirstLineScrollOffset();
+	float PhysicalOffset = 0.f;
+
+	// Handle the first row: it might own a track lane, or reference a track lane owned by a parent
+	// outliner item. In order to correctly position this track lane vertically, we need to accumulate
+	// the item heights from whatever parents we have between the first row and the parent owns our
+	// track lane.
+	{
+		TSharedRef<SOutlinerViewRow> FirstTreeViewRow  = StaticCastSharedRef<SOutlinerViewRow>(ConstCastSharedRef<SWidget>(Children->GetChildAt(0)));
+
+		PhysicalOffset = -(FirstLineScrollOffset * FirstTreeViewRow->GetDesiredSize().Y);
+
+		if (TSharedPtr<STrackLane> FirstRowTrackLane = FirstTreeViewRow->GetTrackLane(false))
+		{
+			FirstRowTrackLane->PositionParentTrackLanes(FirstTreeViewRow->GetDataModel(), PhysicalOffset);
+		}
+	}
+
+	// Now we position any track lanes owned by the outliner items that are visible given the current scrolling.
+	for (int32 Index = 0; Index < NumSlots; ++Index)
+	{
+		TSharedRef<const SOutlinerViewRow> TreeViewRow = StaticCastSharedRef<const SOutlinerViewRow>(Children->GetChildAt(Index));
+
+		// Only get owned track lanes (no parent references... we don't want to move parent lanes down
+		// to children items)
+		if (TSharedPtr<STrackLane> TrackLane = TreeViewRow->GetTrackLane(true))
+		{
+			TrackLane->SetVerticalPosition(PhysicalOffset);
+		}
+
+		TViewModelPtr<IOutlinerExtension> DataModel = TreeViewRow->GetDataModel();
+		if (DataModel)
+		{
+			TSharedPtr<IGeometryExtension> GeometryExtension = DataModel.ImplicitCast();
+
+			const FVector2D TreeViewRowSize = TreeViewRow->GetDesiredSize();
+			const float TreeViewRowHeight = TreeViewRowSize.Y > 0 ? TreeViewRowSize.Y : TreeViewRow->GetDesiredSize().Y;
+			const FVirtualGeometry VirtualGeometry = GeometryExtension ? GeometryExtension->GetVirtualGeometry() : FVirtualGeometry();
+
+			PhysicalNodes.Add(FCachedGeometry{
+				DataModel,
+				PhysicalOffset,
+				TreeViewRowHeight,
+				VirtualGeometry.GetTop(),
+				VirtualGeometry.GetHeight(),
+				VirtualGeometry.GetNestedHeight()
+			});
+		}
+
+		PhysicalOffset += TreeViewRow->GetDesiredSize().Y;
+	}
+}
+
+bool ShouldExpand(const TArray<TWeakViewModelPtr<IOutlinerExtension>>& Selection, ETreeRecursion Recursion, TSharedPtr<FOutlinerViewModel> Outliner)
+{
+	bool bAllExpanded = true;
+	for (TWeakViewModelPtr<IOutlinerExtension> WeakItem : Selection)
+	{
+		// @todo_sequencer_mvvm: Do we need to check selection state here if it's already in a supposed selection set??
+		TViewModelPtr<IOutlinerExtension> SelectedItem = WeakItem.Pin();
+		if (!SelectedItem || !SelectedItem->IsExpanded())
+		{
+			return true;
+		}
+
+		if (Recursion == ETreeRecursion::Recursive)
+		{
+			for (const TViewModelPtr<IOutlinerExtension>& Child : SelectedItem.AsModel()->GetChildrenOfType<IOutlinerExtension>())
+			{
+				if (!Child->IsExpanded())
+				{
+					return true;
+				}
+			}
+		}
+	}
+	return false;
+}
+
+void SOutlinerView::ToggleExpandCollapseNodes(ETreeRecursion Recursion, bool bExpandAll, bool bCollapseAll)
+{
+	TSharedPtr<FOutlinerViewModel> Outliner = WeakOutliner.Pin();
+	if (!Outliner)
+	{
+		return;
+	}
+
+	TArray<TWeakViewModelPtr<IOutlinerExtension>> Selection = GetSelectedItems();
+
+	if (Selection.Num() > 0 && !bExpandAll && !bCollapseAll)
+	{
+		const bool bExpand = ShouldExpand(Selection, Recursion, Outliner);
+		
+		for (TWeakViewModelPtr<IOutlinerExtension> WeakItem : Selection)
+		{
+			TViewModelPtr<IOutlinerExtension> Item = WeakItem.Pin();
+			if (Item)
+			{
+				ExpandCollapseNode(Item, bExpand, Recursion);
+			}
+		}
+	}
+	else if (ItemsSource->Num() > 0)
+	{
+		bool bExpand = ShouldExpand(*ItemsSource, Recursion, Outliner);
+
+		if (bExpandAll)
+		{
+			bExpand = true;
+		}
+		if (bCollapseAll)
+		{
+			bExpand = false;
+		}
+
+		for (TWeakViewModelPtr<IOutlinerExtension> WeakItem : *ItemsSource)
+		{
+			if (TViewModelPtr<IOutlinerExtension> Item = WeakItem.Pin())
+			{
+				ExpandCollapseNode(Item, bExpand, Recursion);
+			}
+		}
+	}
+}
+
+void SOutlinerView::ExpandCollapseNode(TViewModelPtr<IOutlinerExtension> InViewModel, bool bExpansionState, ETreeRecursion Recursion)
+{
+	SetItemExpansion(InViewModel, bExpansionState);
+
+	if (Recursion == ETreeRecursion::Recursive)
+	{
+		for (TViewModelPtr<IOutlinerExtension> Child : InViewModel.AsModel()->GetChildrenOfType<IOutlinerExtension>())
+		{
+			ExpandCollapseNode(Child, bExpansionState, ETreeRecursion::Recursive);
+		}
+	}
+}
+
+void SOutlinerView::OnExpansionChanged(TWeakViewModelPtr<IOutlinerExtension> InWeakItem, bool bIsExpanded)
+{
+	TViewModelPtr<IOutlinerExtension> SelectedItem = InWeakItem.Pin();
+	if (!SelectedItem)
+	{
+		return;
+	}
+
+	SelectedItem->SetExpansion(bIsExpanded);
+
+	// Expand any children that are also expanded
+	for (const TViewModelPtr<IOutlinerExtension>& Child : SelectedItem.AsModel()->GetChildrenOfType<IOutlinerExtension>())
+	{
+		if (Child->IsExpanded())
+		{
+			SetItemExpansion(Child, true);
+		}
+	}
+}
+
+void SOutlinerView::SetItemExpansionRecursive(TWeakViewModelPtr<IOutlinerExtension> InItem, bool bIsExpanded)
+{
+	if (TViewModelPtr<IOutlinerExtension> Item = InItem.Pin())
+	{
+		ExpandCollapseNode(Item, bIsExpanded, ETreeRecursion::Recursive);
+	}
+}
+
+void SOutlinerView::OnGetChildren(TWeakViewModelPtr<IOutlinerExtension> InParent, TArray<TWeakViewModelPtr<IOutlinerExtension>>& OutChildren) const
+{
+	TSharedPtr<FOutlinerViewModel> Outliner  = WeakOutliner.Pin();
+	TViewModelPtr<IOutlinerExtension> DataModel = InParent.Pin();
+
+	if (Outliner && DataModel)
+	{
+		for (TViewModelPtr<IOutlinerExtension> Child : DataModel.AsModel()->GetChildrenOfType<IOutlinerExtension>())
+		{
+			if (!Child->IsFilteredOut())
+			{
+				OutChildren.Add(Child);
+			}
+		}
+	}
+}
+
+TSharedPtr<STrackLane> SOutlinerView::FindOrCreateParentLane(TViewModelPtr<IOutlinerExtension> InDataModel)
+{
+	TSharedPtr<FOutlinerViewModel> Outliner  = WeakOutliner.Pin();
+	if (!Outliner)
+	{
+		return nullptr;
+	}
+
+	// Find any parent that will create a nested track lane
+	for (const TViewModelPtr<ITrackAreaExtension>& ParentTrackArea : InDataModel.AsModel()->GetAncestorsOfType<ITrackAreaExtension>())
+	{
+		FTrackAreaParameters TrackAreaParameters = ParentTrackArea->GetTrackAreaParameters();
+		if (TrackAreaParameters.LaneType == ETrackAreaLaneType::Nested)
+		{
+			// Find the first outliner extension parent (should be the same as ParentTrackArea really)
+			for (const TViewModelPtr<IOutlinerExtension>& OutlinerExtItem : ParentTrackArea.AsModel()->GetAncestorsOfType<IOutlinerExtension>(true))
+			{
+				TSharedPtr<STrackLane> TrackLane = TrackArea->FindTrackSlot(OutlinerExtItem);
+				if (!TrackLane)
+				{
+					// Add a track slot for the nested track lane parent if it doesn't already exist
+					TrackLane = SNew(STrackLane, TrackArea->GetViewModel(), OutlinerExtItem, TrackAreaParameters, SharedThis(this));
+					TrackArea->AddTrackSlot(OutlinerExtItem, TrackLane);
+				}
+
+				return TrackLane;
+			}
+
+			// Once we find any track area extension that wants nested lanes, return if it didn't manage to create one
+			break;
+		}
+	}
+
+	return nullptr;
+}
+
+void SOutlinerView::CreateTrackLanesForRow(TSharedRef<SOutlinerViewRow> InRow, TViewModelPtr<IOutlinerExtension> InDataModel)
+{
+	// Create a lane for this item, if needed, or find a lane to use from one of its parents.
+	// We also create a lane on the appropriate parent if it doesn't exist, because we could be in a situation
+	// where we jump-scrolled directly to the given row and item without ever having had the parent in view.
+	const bool bIncludeDataModel = true;
+	for (TViewModelPtr<IOutlinerExtension> CurDataModel : InDataModel.AsModel()->GetAncestors(bIncludeDataModel))
+	{
+		if (TSharedPtr<ITrackAreaExtension> TrackAreaExtension = CurDataModel.ImplicitCast())
+		{
+			FTrackAreaParameters Parameters = TrackAreaExtension->GetTrackAreaParameters();
+			if (Parameters.LaneType != ETrackAreaLaneType::None)
+			{
+				TSharedPtr<STrackLane> TrackLane = TrackArea->FindTrackSlot(CurDataModel);
+				if (!TrackLane)
+				{
+					// Add a track slot for the row
+					TrackLane = SNew(STrackLane, TrackArea->GetViewModel(), CurDataModel, Parameters, SharedThis(this));
+
+					TSharedPtr<STrackLane> ParentLane = FindOrCreateParentLane(CurDataModel);
+					if (ParentLane)
+					{
+						TrackLane->SetParent(ParentLane);
+					}
+
+					TrackArea->AddTrackSlot(CurDataModel, TrackLane);
+				}
+				ensure(TrackLane);
+				InRow->SetTrackLane(TrackLane);
+				break;
+			}
+		}
+	}
+}
+
+TSharedRef<ITableRow> SOutlinerView::OnGenerateRow(TWeakViewModelPtr<IOutlinerExtension> InWeakModel, const TSharedRef<STableViewBase>& OwnerTable)
+{
+	TSharedRef<SOutlinerViewRow> Row =
+		SNew(SOutlinerViewRow, OwnerTable, InWeakModel)
+		.OnDetectDrag(this, &SOutlinerView::OnDragRow)
+		.OnGenerateWidgetForColumn(this, &SOutlinerView::GenerateWidgetForColumn);
+
+	if (TViewModelPtr<IOutlinerExtension> ViewModel = InWeakModel.Pin())
+	{
+		CreateTrackLanesForRow(Row, ViewModel);
+	}
+	return Row;
+}
+
+TSharedRef<SWidget> SOutlinerView::GenerateWidgetForColumn(TViewModelPtr<IOutlinerExtension> InDataModel, const FName& ColumnId, const TSharedRef<SOutlinerViewRow>& Row) const
+{
+	const FOutlinerViewColumn* Definition = Columns.Find(ColumnId);
+
+	if (ensureMsgf(Definition, TEXT("Invalid column name specified")))
+	{
+		return Definition->Generator(InDataModel, Row);
+	}
+
+	return SNullWidget::NullWidget;
+}
+
+} // namespace Sequencer
+} // namespace UE
+
diff --git a/Engine/Source/Editor/SequencerCore/Private/MVVM/Views/ViewUtilities.cpp b/Engine/Source/Editor/SequencerCore/Private/MVVM/Views/ViewUtilities.cpp
new file mode 100644
--- /dev/null
+++ b/Engine/Source/Editor/SequencerCore/Private/MVVM/Views/ViewUtilities.cpp
@@ -0,0 +1,82 @@
+// Copyright Epic Games, Inc. All Rights Reserved.
+
+#include "MVVM/Views/ViewUtilities.h"
+#include "Layout/Margin.h"
+#include "Fonts/SlateFontInfo.h"
+#include "Widgets/DeclarativeSyntaxSupport.h"
+#include "Widgets/SBoxPanel.h"
+#include "Widgets/Images/SImage.h"
+#include "Widgets/Text/STextBlock.h"
+#include "Widgets/Input/SComboButton.h"
+#include "Widgets/Input/SButton.h"
+#include "Styling/CoreStyle.h"
+#include "Styling/AppStyle.h"
+
+namespace UE
+{
+namespace Sequencer
+{
+
+TSharedRef<SWidget> MakeAddButton(FText HoverText, FOnGetContent MenuContent, const TAttribute<bool>& HoverState, const TAttribute<bool>& IsEnabled)
+{
+	auto GetRolloverVisibility = [](TAttribute<bool> InHoverState, TWeakPtr<SComboButton> WeakComboButton)
+	{
+		TSharedPtr<SComboButton> TmpComboButton = WeakComboButton.Pin();
+		if (InHoverState.Get() || TmpComboButton->IsOpen())
+		{
+			return EVisibility::SelfHitTestInvisible;
+		}
+		else
+		{
+			return EVisibility::Collapsed;
+		}
+	};
+
+	FSlateFontInfo SmallLayoutFont = FCoreStyle::GetDefaultFontStyle("Regular", 8);
+
+	TSharedRef<STextBlock> ComboButtonText = SNew(STextBlock)
+		.Text(HoverText)
+		.Font(SmallLayoutFont)
+		.ColorAndOpacity( FSlateColor::UseForeground() );
+
+	TSharedRef<SComboButton> ComboButton =
+
+		SNew(SComboButton)
+		.HasDownArrow(false)
+		.IsFocusable(true)
+		.ButtonStyle(FAppStyle::Get(), "HoverHintOnly")
+		.ForegroundColor( FSlateColor::UseForeground() )
+		.IsEnabled(IsEnabled)
+		.OnGetMenuContent(MenuContent)
+		.ContentPadding(FMargin(5, 2))
+		.HAlign(HAlign_Center)
+		.VAlign(VAlign_Center)
+		.ButtonContent()
+		[
+			SNew(SHorizontalBox)
+
+			+ SHorizontalBox::Slot()
+			.AutoWidth()
+			.VAlign(VAlign_Center)
+			.Padding(FMargin(0,0,2,0))
+			[
+				SNew(SImage)
+				.ColorAndOpacity( FSlateColor::UseForeground() )
+				.Image(FAppStyle::GetBrush("Plus"))
+			]
+
+			+ SHorizontalBox::Slot()
+			.VAlign(VAlign_Center)
+			.AutoWidth()
+			[
+				ComboButtonText
+			]
+		];
+
+	ComboButtonText->SetVisibility(MakeAttributeLambda(GetRolloverVisibility, HoverState, TWeakPtr<SComboButton>(ComboButton)));
+
+	return ComboButton;
+}
+
+} // namespace Sequencer
+} // namespace UE
\ No newline at end of file
diff --git a/Engine/Source/Editor/SequencerCore/Private/SequencerCoreModule.h b/Engine/Source/Editor/SequencerCore/Private/SequencerCoreModule.h
new file mode 100644
--- /dev/null
+++ b/Engine/Source/Editor/SequencerCore/Private/SequencerCoreModule.h
@@ -0,0 +1,24 @@
+// Copyright Epic Games, Inc. All Rights Reserved.
+
+#pragma once
+
+#include "Modules/ModuleInterface.h"
+
+namespace UE
+{
+namespace Sequencer
+{
+
+
+/**
+ * Interface for the Sequencer module.
+ */
+class ISequencerCoreModule
+	: public IModuleInterface
+{
+public:
+};
+
+
+} // namespace Sequencer
+} // namespace UE
\ No newline at end of file
diff --git a/Engine/Source/Editor/SequencerCore/Public/ISequencerInputHandler.h b/Engine/Source/Editor/SequencerCore/Public/ISequencerInputHandler.h
new file mode 100644
--- /dev/null
+++ b/Engine/Source/Editor/SequencerCore/Public/ISequencerInputHandler.h
@@ -0,0 +1,19 @@
+// Copyright Epic Games, Inc. All Rights Reserved.
+
+#pragma once
+
+#include "CoreTypes.h"
+#include "Input/Reply.h"
+
+class SWidget;
+
+/**
+ * Common base-class for objects that handle input in the sequencer.
+ */
+struct ISequencerInputHandler
+{
+	virtual FReply OnMouseButtonDown(SWidget& OwnerWidget, const FGeometry& MyGeometry, const FPointerEvent& MouseEvent) = 0;
+	virtual FReply OnMouseButtonUp(SWidget& OwnerWidget, const FGeometry& MyGeometry, const FPointerEvent& MouseEvent) = 0;
+	virtual FReply OnMouseMove(SWidget& OwnerWidget, const FGeometry& MyGeometry, const FPointerEvent& MouseEvent) = 0;
+	virtual FReply OnMouseWheel(SWidget& OwnerWidget, const FGeometry& MyGeometry, const FPointerEvent& MouseEvent) = 0;
+};
diff --git a/Engine/Source/Editor/SequencerCore/Public/MVVM/Extensions/DynamicExtensionContainer.h b/Engine/Source/Editor/SequencerCore/Public/MVVM/Extensions/DynamicExtensionContainer.h
new file mode 100644
--- /dev/null
+++ b/Engine/Source/Editor/SequencerCore/Public/MVVM/Extensions/DynamicExtensionContainer.h
@@ -0,0 +1,216 @@
+// Copyright Epic Games, Inc. All Rights Reserved.
+
+#pragma once
+
+#include "Containers/SortedMap.h"
+#include "Delegates/Delegate.h"
+#include "Delegates/DelegateCombinations.h"
+#include "MVVM/ICastable.h"
+#include "MVVM/ViewModelTypeID.h"
+#include "Templates/SharedPointer.h"
+#include "Misc/InlineValue.h"
+
+namespace UE
+{
+namespace Sequencer
+{
+
+class FDynamicExtensionContainer;
+class FViewModel;
+
+/**
+ * Base class for dynamic extensions that can be added to sequence data models and participate in dynamic casting.
+ */
+class SEQUENCERCORE_API IDynamicExtension : public TSharedFromThis<IDynamicExtension>
+{
+public:
+
+	UE_SEQUENCER_DECLARE_VIEW_MODEL_TYPE_ID(IDynamicExtension)
+
+	virtual ~IDynamicExtension(){}
+
+	/** Called when the extension is created on a data model */
+	virtual void OnCreated(TSharedRef<FViewModel> InWeakOwner) {}
+	/** Called after all extensions have been created on a data model */
+	virtual void OnPostInitialize() { OnReinitialize(); }
+	/** Called to reinitialize extensions after a major change */
+	virtual void OnReinitialize() {}
+};
+
+DECLARE_DELEGATE_OneParam(FDynamicExtensionCallback, TSharedRef<FViewModel>);
+
+/**
+ * Base class for supporting dynamic extensions that participate in dynamic casting.
+ */
+class SEQUENCERCORE_API FDynamicExtensionContainer
+{
+public:
+
+	virtual ~FDynamicExtensionContainer() {}
+
+	FDynamicExtensionContainer() = default;
+
+	FDynamicExtensionContainer(const FDynamicExtensionContainer&) = delete;
+	FDynamicExtensionContainer& operator=(const FDynamicExtensionContainer&) = delete;
+
+public:
+
+	UE_SEQUENCER_DECLARE_VIEW_MODEL_TYPE_ID(FDynamicExtensionContainer)
+
+	template<typename T>
+	const T* CastDynamic() const
+	{
+		const void* Result = CastDynamic(T::ID);
+		return static_cast<const T*>(Result);
+	}
+
+	template<typename T>
+	const T* CastDynamicChecked() const
+	{
+		const void* Result = CastDynamicChecked(T::ID);
+		return static_cast<const T*>(Result);
+	}
+
+	template<typename T>
+	T* CastDynamic()
+	{
+		return const_cast<T*>(const_cast<const FDynamicExtensionContainer*>(this)->CastDynamic<T>());
+	}
+
+	template<typename T>
+	T* CastDynamicChecked()
+	{
+		return const_cast<T*>(const_cast<const FDynamicExtensionContainer*>(this)->CastDynamicChecked<T>());
+	}
+
+	const void* CastDynamic(FViewModelTypeID Type) const
+	{
+		for (const TInlineValue<FDynamicExtensionInfo>& DynamicExtension : DynamicExtensions)
+		{
+			if (const void* Result = DynamicExtension->CastDynamic(Type))
+			{
+				return Result;
+			}
+		}
+		return nullptr;
+	}
+
+	const void* CastDynamicChecked(FViewModelTypeID Type) const
+	{
+		const void* Result = CastDynamic(Type);
+		check(Result);
+		return Result;
+	}
+
+	void* CastDynamic(FViewModelTypeID Type)
+	{
+		return const_cast<void*>(const_cast<const FDynamicExtensionContainer*>(this)->CastDynamic(Type));
+	}
+
+	void* CastDynamicChecked(FViewModelTypeID Type)
+	{
+		void* Result = CastDynamic(Type);
+		check(Result);
+		return Result;
+	}
+
+protected:
+
+	template<typename T, typename... InArgTypes>
+	T& AddDynamicExtension(TSharedRef<FViewModel> InOwner, InArgTypes&&... Args)
+	{
+		T* Existing = CastDynamic<T>();
+		if (Existing)
+		{
+			return static_cast<T&>(*Existing);
+		}
+
+		static_assert(sizeof(FDynamicExtensionInfo) == sizeof(TDynamicExtensionInfo<T>), "Typed extension wrapper can't contain extra data.");
+
+		TSharedRef<T> NewExtension = MakeShared<T>(Forward<InArgTypes>(Args)...);
+		DynamicExtensions.Add(TDynamicExtensionInfo<T>(NewExtension));
+		NewExtension->OnCreated(InOwner);
+		return NewExtension.Get();
+	}
+
+	void PostInitializeExtensions()
+	{
+		for (const TInlineValue<FDynamicExtensionInfo>& DynamicExtension : DynamicExtensions)
+		{
+			check(DynamicExtension.IsValid());
+			DynamicExtension->Extension->OnPostInitialize();
+		}
+	}
+
+	void ReinitializeExtensions()
+	{
+		for (const TInlineValue<FDynamicExtensionInfo>& DynamicExtension : DynamicExtensions)
+		{
+			check(DynamicExtension.IsValid());
+			DynamicExtension->Extension->OnReinitialize();
+		}
+	}
+
+public:
+
+	struct Implements{};
+
+private:
+
+	struct FDynamicExtensionInfo
+	{
+		TSharedRef<IDynamicExtension> Extension;
+		FDynamicExtensionInfo(TSharedRef<IDynamicExtension> InExtension) : Extension(InExtension) {}
+		virtual ~FDynamicExtensionInfo() {}
+		virtual const void* CastDynamic(FViewModelTypeID Type) const { return nullptr; }
+	};
+
+	template<typename T>
+	struct TDynamicExtensionInfo : FDynamicExtensionInfo
+	{
+		TDynamicExtensionInfo(TSharedRef<T> InExtension) 
+			: FDynamicExtensionInfo(StaticCastSharedRef<IDynamicExtension>(InExtension))
+		{
+		}
+		virtual const void* CastDynamic(FViewModelTypeID Type) const
+		{
+			const void* Result = nullptr;
+			TSharedRef<T> TypedExtension = StaticCastSharedRef<T>(Extension);
+			FCastableBoilerplate::CastImplementation(&TypedExtension.Get(), Type, Result);
+			return Result;
+		}
+	};
+
+	TArray<TInlineValue<FDynamicExtensionInfo>> DynamicExtensions;
+};
+
+template<>
+struct TCompositeCast<FDynamicExtensionContainer::Implements>
+{
+	static void Apply(const FDynamicExtensionContainer* This, FViewModelTypeID ToType, const void*& OutResult)
+	{
+		if (OutResult == nullptr)
+		{
+			if (const void* Result = FCastableBoilerplate::DirectCast(This, ToType))
+			{
+				OutResult = Result;
+				return;
+			}
+
+			if (const void* Result = This->CastDynamic(ToType))
+			{
+				OutResult = Result;
+			}
+		}
+	}
+
+	template<typename T>
+	static bool IsAny(const ICastable* This)
+	{
+		return T::ID == FDynamicExtensionContainer::ID;
+	}
+};
+
+} // namespace Sequencer
+} // namespace UE
+
diff --git a/Engine/Source/Editor/SequencerCore/Public/MVVM/Extensions/IDeletableExtension.h b/Engine/Source/Editor/SequencerCore/Public/MVVM/Extensions/IDeletableExtension.h
new file mode 100644
--- /dev/null
+++ b/Engine/Source/Editor/SequencerCore/Public/MVVM/Extensions/IDeletableExtension.h
@@ -0,0 +1,26 @@
+// Copyright Epic Games, Inc. All Rights Reserved.
+
+#pragma once
+
+#include "CoreTypes.h"
+#include "MVVM/ViewModelTypeID.h"
+
+namespace UE
+{
+namespace Sequencer
+{
+
+class SEQUENCERCORE_API IDeletableExtension
+{
+public:
+
+	UE_SEQUENCER_DECLARE_VIEW_MODEL_TYPE_ID(IDeletableExtension)
+
+	virtual ~IDeletableExtension(){}
+
+	virtual bool CanDelete(FText* OutErrorMessage) const = 0;
+	virtual void Delete() = 0;
+};
+
+} // namespace Sequencer
+} // namespace UE
diff --git a/Engine/Source/Editor/SequencerCore/Public/MVVM/Extensions/IDimmableExtension.h b/Engine/Source/Editor/SequencerCore/Public/MVVM/Extensions/IDimmableExtension.h
new file mode 100644
--- /dev/null
+++ b/Engine/Source/Editor/SequencerCore/Public/MVVM/Extensions/IDimmableExtension.h
@@ -0,0 +1,26 @@
+// Copyright Epic Games, Inc. All Rights Reserved.
+
+#pragma once
+
+#include "CoreMinimal.h"
+#include "MVVM/ViewModelTypeID.h"
+
+namespace UE
+{
+namespace Sequencer
+{
+
+class SEQUENCERCORE_API IDimmableExtension
+{
+public:
+
+	UE_SEQUENCER_DECLARE_VIEW_MODEL_TYPE_ID(IDimmableExtension)
+
+	virtual ~IDimmableExtension(){}
+
+	virtual bool IsDimmed() const = 0;
+};
+
+} // namespace Sequencer
+} // namespace UE
+
diff --git a/Engine/Source/Editor/SequencerCore/Public/MVVM/Extensions/IDraggableOutlinerExtension.h b/Engine/Source/Editor/SequencerCore/Public/MVVM/Extensions/IDraggableOutlinerExtension.h
new file mode 100644
--- /dev/null
+++ b/Engine/Source/Editor/SequencerCore/Public/MVVM/Extensions/IDraggableOutlinerExtension.h
@@ -0,0 +1,33 @@
+// Copyright Epic Games, Inc. All Rights Reserved.
+
+#pragma once
+
+#include "MVVM/ViewModelTypeID.h"
+#include "Misc/FrameNumber.h"
+
+struct FGuid;
+
+namespace UE
+{
+namespace Sequencer
+{
+
+
+/**
+ * Extension for models that can be dragged
+ */
+class SEQUENCERCORE_API IDraggableOutlinerExtension
+{
+public:
+
+	UE_SEQUENCER_DECLARE_VIEW_MODEL_TYPE_ID(IDraggableOutlinerExtension)
+
+	virtual ~IDraggableOutlinerExtension(){}
+
+	/** Returns whether the model can be dragged */
+	virtual bool CanDrag() const = 0;
+};
+
+} // namespace Sequencer
+} // namespace UE
+
diff --git a/Engine/Source/Editor/SequencerCore/Public/MVVM/Extensions/IGroupableExtension.h b/Engine/Source/Editor/SequencerCore/Public/MVVM/Extensions/IGroupableExtension.h
new file mode 100644
--- /dev/null
+++ b/Engine/Source/Editor/SequencerCore/Public/MVVM/Extensions/IGroupableExtension.h
@@ -0,0 +1,36 @@
+// Copyright Epic Games, Inc. All Rights Reserved.
+
+#pragma once
+
+#include "MVVM/ViewModelTypeID.h"
+#include "Containers/StringFwd.h"
+
+struct FGuid;
+
+namespace UE
+{
+namespace Sequencer
+{
+
+class SEQUENCERCORE_API IGroupableExtension
+{
+public:
+
+	UE_SEQUENCER_DECLARE_VIEW_MODEL_TYPE_ID(IGroupableExtension)
+
+	virtual ~IGroupableExtension(){}
+
+	virtual void GetIdentifierForGrouping(TStringBuilder<128>& OutString) const = 0;
+
+	FString GetIdentifierForGrouping() const
+	{
+		TStringBuilder<128> Builder;
+		GetIdentifierForGrouping(Builder);
+		return Builder.ToString();
+	}
+};
+
+
+} // namespace Sequencer
+} // namespace UE
+
diff --git a/Engine/Source/Editor/SequencerCore/Public/MVVM/Extensions/IHotspotExtension.h b/Engine/Source/Editor/SequencerCore/Public/MVVM/Extensions/IHotspotExtension.h
new file mode 100644
--- /dev/null
+++ b/Engine/Source/Editor/SequencerCore/Public/MVVM/Extensions/IHotspotExtension.h
@@ -0,0 +1,24 @@
+// Copyright Epic Games, Inc. All Rights Reserved.
+
+#pragma once
+
+#include "MVVM/ViewModelTypeID.h"
+#include "Templates/SharedPointer.h"
+
+namespace UE
+{
+namespace Sequencer
+{
+
+class SEQUENCERCORE_API IHotspotExtension
+{
+public:
+
+	UE_SEQUENCER_DECLARE_VIEW_MODEL_TYPE_ID(IHotspotExtension)
+
+	virtual ~IHotspotExtension(){}
+};
+
+} // namespace Sequencer
+} // namespace UE
+
diff --git a/Engine/Source/Editor/SequencerCore/Public/MVVM/Extensions/IMutableExtension.h b/Engine/Source/Editor/SequencerCore/Public/MVVM/Extensions/IMutableExtension.h
new file mode 100644
--- /dev/null
+++ b/Engine/Source/Editor/SequencerCore/Public/MVVM/Extensions/IMutableExtension.h
@@ -0,0 +1,30 @@
+// Copyright Epic Games, Inc. All Rights Reserved.
+
+#pragma once
+
+#include "CoreMinimal.h"
+#include "MVVM/ViewModelTypeID.h"
+
+namespace UE
+{
+namespace Sequencer
+{
+
+/**
+ * An extension for outliner nodes that can be muted
+ */
+class SEQUENCERCORE_API IMutableExtension
+{
+public:
+
+	UE_SEQUENCER_DECLARE_VIEW_MODEL_TYPE_ID(IMutableExtension)
+
+	virtual ~IMutableExtension(){}
+
+	/** Returns whether this item is muted, or is in a muted branch */
+	virtual bool IsMuted() const = 0;
+};
+
+} // namespace Sequencer
+} // namespace UE
+
diff --git a/Engine/Source/Editor/SequencerCore/Public/MVVM/Extensions/IOutlinerExtension.h b/Engine/Source/Editor/SequencerCore/Public/MVVM/Extensions/IOutlinerExtension.h
new file mode 100644
--- /dev/null
+++ b/Engine/Source/Editor/SequencerCore/Public/MVVM/Extensions/IOutlinerExtension.h
@@ -0,0 +1,225 @@
+// Copyright Epic Games, Inc. All Rights Reserved.
+
+#pragma once
+
+#include "Templates/SharedPointer.h"
+#include "MVVM/ICastable.h"
+#include "MVVM/ViewModelTypeID.h"
+#include "Widgets/Views/STableRow.h"
+#include "Widgets/Views/SListView.h"
+#include "MVVM/ViewModels/ViewModel.h"
+#include "MVVM/Views/TreeViewTraits.h"
+
+struct FSlateColor;
+struct FSlateBrush;
+struct FSlateFontInfo;
+
+class FText;
+class SWidget;
+class ISequencer;
+
+namespace UE
+{
+namespace Sequencer
+{
+
+class IOutlinerExtension;
+class FOutlinerViewModel;
+class FViewModel;
+class FEditorViewModel;
+
+enum class EOutlinerSelectionState
+{
+	None                      = 0,
+	SelectedDirectly          = 1 << 1,
+	HasSelectedKeys           = 1 << 2,
+	HasSelectedTrackAreaItems = 1 << 3,
+};
+ENUM_CLASS_FLAGS(EOutlinerSelectionState);
+
+struct SEQUENCERCORE_API FOutlinerSizing
+{
+	FOutlinerSizing() = default;
+
+	FOutlinerSizing(float InHeight)
+		: Height(InHeight)
+	{}
+
+	FOutlinerSizing(float InHeight, float UniformPadding)
+		: Height(InHeight)
+		, PaddingTop(UniformPadding)
+		, PaddingBottom(UniformPadding)
+	{}
+
+	float GetTotalHeight() const
+	{
+		return Height + PaddingTop + PaddingBottom;
+	}
+
+	float Height = 0.f;
+	float PaddingTop = 0.f;
+	float PaddingBottom = 0.f;
+};
+
+/** Facade class to allow SOutlinerViewRow to be passed around publicly */
+class SEQUENCERCORE_API ISequencerTreeViewRow : public SMultiColumnTableRow<TWeakViewModelPtr<IOutlinerExtension>>
+{
+};
+
+/** Parameters for creating an outliner item widget. */
+struct SEQUENCERCORE_API FCreateOutlinerViewParams
+{
+	FCreateOutlinerViewParams(const TSharedRef<ISequencerTreeViewRow> InTreeViewRow, const TSharedPtr<FEditorViewModel> InEditor)
+		: TreeViewRow(InTreeViewRow), Editor(InEditor)
+	{}
+
+	const TSharedRef<ISequencerTreeViewRow> TreeViewRow;
+	const TSharedPtr<FEditorViewModel> Editor;
+};
+
+/** Parameters for building a context menu widget */
+struct SEQUENCERCORE_API FCreateOutlinerContextMenuWidgetParams
+{
+	FCreateOutlinerContextMenuWidgetParams(const TSharedPtr<FEditorViewModel> InEditor)
+		: Editor(InEditor)
+	{}
+
+	const TSharedPtr<FEditorViewModel> Editor;
+};
+
+/**
+ * Extension interface for models that can be displayed in the outliner view.
+ */
+class SEQUENCERCORE_API IOutlinerExtension
+{
+public:
+
+	UE_SEQUENCER_DECLARE_VIEW_MODEL_TYPE_ID(IOutlinerExtension)
+
+	virtual ~IOutlinerExtension(){}
+
+	/** Return the identifier of the owning model */
+	virtual FName GetIdentifier() const = 0;
+
+	/** Return the desired size for the owning model */
+	virtual FOutlinerSizing GetOutlinerSizing() const = 0;
+
+	/** Gets whether this outliner item is expanded */
+	virtual bool IsExpanded() const = 0;
+
+	/** Sets the expansion state of this outliner item */
+	virtual void SetExpansion(bool bIsExpanded) = 0;
+
+	/** Gets whether this outliner item has been filtered out */
+	virtual bool IsFilteredOut() const = 0;
+
+	/** Sets whether this outliner item is being filtered out */
+	virtual void SetFilteredOut(bool bIsFilteredOut) = 0;
+
+	/** Gets the outliner item's selection state */
+	virtual EOutlinerSelectionState GetSelectionState() const = 0;
+
+	/** Sets the outliner item's selection state */
+	virtual void SetSelectionState(EOutlinerSelectionState InState) = 0;
+
+	/** Toggles a bit in the outliner item's selection state */
+	virtual void ToggleSelectionState(EOutlinerSelectionState InState, bool bInValue);
+
+	/** Create the visual outliner widget for the owning model */
+	virtual TSharedRef<SWidget> CreateOutlinerView(const FCreateOutlinerViewParams& InParams) = 0;
+
+	/** Build the context menu widget for this outliner item */
+	virtual TSharedPtr<SWidget> CreateContextMenuWidget(const FCreateOutlinerContextMenuWidgetParams& InParams) = 0;
+
+	/** Gets the label for the outliner item */
+	virtual FText GetLabel() const;
+
+	/** Gets the font for the outliner item */
+	virtual FSlateFontInfo GetLabelFont() const;
+
+	/** Gets the color for the outliner item label */
+	virtual FSlateColor GetLabelColor() const;
+
+	/** Gets the tooltip text for the outliner item */
+	virtual FText GetLabelToolTipText() const;
+
+	/** Gets the icon for the outliner item */
+	virtual const FSlateBrush* GetIconBrush() const;
+
+	/** Gets the icon tint color for the outliner item */
+	virtual FSlateColor GetIconTint() const;
+
+	/** Gets the overlay icon for the outliner item */
+	virtual const FSlateBrush* GetIconOverlayBrush() const;
+
+	/** Gets the tooltip text for the outliner item icon */
+	virtual FText GetIconToolTipText() const;
+
+	/** Gets whether the outliner item has a background */
+	virtual bool HasBackground() const;
+
+	/** Utility method to get the full path name of the given model */
+	static FString GetPathName(const FViewModel& Item);
+	/** Utility method to get the full path name of the given model */
+	static FString GetPathName(const TSharedPtr<FViewModel> Item);
+	/** Utility method to get the full path name of the given model */
+	static FString GetPathName(const TSharedPtr<const FViewModel> Item);
+	/** Utility method to get the full path name of the given model */
+	static FString GetPathName(const TWeakPtr<const FViewModel> Item);
+	/** Utility method to get the full path name of the given model */
+	static void GetPathName(const FViewModel& Item, FStringBuilderBase& OutString);
+};
+
+/**
+ * Default implementation for some of the outliner extension API.
+ */
+class SEQUENCERCORE_API FOutlinerExtensionShim : public IOutlinerExtension
+{
+public:
+
+	virtual bool IsExpanded() const override { return bIsExpanded; }
+	virtual void SetExpansion(bool bInIsExpanded) override { bIsExpanded = bInIsExpanded; }
+	virtual bool IsFilteredOut() const override { return bIsFilteredOut; }
+	virtual void SetFilteredOut(bool bInIsFilteredOut) override { bIsFilteredOut = bInIsFilteredOut; }
+	virtual EOutlinerSelectionState GetSelectionState() const override { return SelectionState; }
+	virtual void SetSelectionState(EOutlinerSelectionState InState) override { SelectionState = InState; }
+
+protected:
+
+	bool bIsExpanded = false;
+	bool bIsFilteredOut = false;
+	EOutlinerSelectionState SelectionState = EOutlinerSelectionState::None;
+};
+
+/**
+ * Extension for outliner items that accept drag and drop.
+ */
+class SEQUENCERCORE_API IOutlinerDropTargetOutlinerExtension
+{
+public:
+
+	UE_SEQUENCER_DECLARE_VIEW_MODEL_TYPE_ID(IOutlinerDropTargetOutlinerExtension)
+
+	virtual ~IOutlinerDropTargetOutlinerExtension(){}
+
+	virtual TOptional<EItemDropZone> CanAcceptDrop(const FViewModelPtr& TargetModel, const FDragDropEvent& DragDropEvent, EItemDropZone InItemDropZone) = 0;
+	virtual void PerformDrop(const FViewModelPtr& TargetModel, const FDragDropEvent& DragDropEvent, EItemDropZone InItemDropZone) = 0;
+};
+
+/**
+ * Extension for outliner models that want to compute their and their children's sizing.
+ */
+class SEQUENCERCORE_API ICompoundOutlinerExtension
+{
+public:
+
+	UE_SEQUENCER_DECLARE_VIEW_MODEL_TYPE_ID(ICompoundOutlinerExtension)
+
+	virtual ~ICompoundOutlinerExtension(){}
+
+	virtual FOutlinerSizing RecomputeSizing() = 0;
+};
+
+} // namespace Sequencer
+} // namespace UE
+
diff --git a/Engine/Source/Editor/SequencerCore/Public/MVVM/Extensions/IPinnableExtension.h b/Engine/Source/Editor/SequencerCore/Public/MVVM/Extensions/IPinnableExtension.h
new file mode 100644
--- /dev/null
+++ b/Engine/Source/Editor/SequencerCore/Public/MVVM/Extensions/IPinnableExtension.h
@@ -0,0 +1,62 @@
+// Copyright Epic Games, Inc. All Rights Reserved.
+
+#pragma once
+
+#include "CoreMinimal.h"
+#include "MVVM/ViewModelTypeID.h"
+
+namespace UE
+{
+namespace Sequencer
+{
+
+class FViewModel;
+
+/**
+ * An extension for outliner nodes that can be pinned at the top
+ */
+class SEQUENCERCORE_API IPinnableExtension
+{
+public:
+
+	UE_SEQUENCER_DECLARE_VIEW_MODEL_TYPE_ID(IPinnableExtension)
+
+	virtual ~IPinnableExtension(){}
+
+	/** Returns whether this item is pinned, or is in a pinned branch */
+	virtual bool IsPinned() const = 0;
+	/** Sets whether this item is pinned */
+	virtual void SetPinned(bool bInIsPinned) = 0;
+	/** Lets parent items report to their children that they are in a pinned branch */
+	virtual void ReportPinnedParent(bool bInIsInPinnedBranch) = 0;
+};
+
+/**
+ * Default implementation for IPinnableExtension, with helper method to cache pinned
+ * states for children of a pinned branch
+ */
+class SEQUENCERCORE_API FPinnableExtensionShim : public IPinnableExtension
+{
+public:
+
+	virtual bool IsPinned() const override;
+	virtual void SetPinned(bool bInIsPinned) override;
+	virtual void ReportPinnedParent(bool bInIsInPinnedBranch) override;
+
+public:
+	
+	/**
+	 * Looks for any direct child of RootModel that is pinnable, and reports all their descendants
+	 * that they are in a pinned branch, to make that lookup faster
+	 */
+	static void UpdateCachedPinnedState(TSharedPtr<FViewModel> RootModel);
+
+private:
+
+	bool bIsPinned = false;
+	bool bIsInPinnedBranch = false;
+};
+
+} // namespace Sequencer
+} // namespace UE
+
diff --git a/Engine/Source/Editor/SequencerCore/Public/MVVM/Extensions/IResizableExtension.h b/Engine/Source/Editor/SequencerCore/Public/MVVM/Extensions/IResizableExtension.h
new file mode 100644
--- /dev/null
+++ b/Engine/Source/Editor/SequencerCore/Public/MVVM/Extensions/IResizableExtension.h
@@ -0,0 +1,27 @@
+// Copyright Epic Games, Inc. All Rights Reserved.
+
+#pragma once
+
+#include "CoreMinimal.h"
+#include "MVVM/ViewModelTypeID.h"
+
+namespace UE
+{
+namespace Sequencer
+{
+
+class SEQUENCERCORE_API IResizableExtension
+{
+public:
+
+	UE_SEQUENCER_DECLARE_VIEW_MODEL_TYPE_ID(IResizableExtension)
+
+	virtual ~IResizableExtension(){}
+
+	virtual bool IsResizable() const = 0;
+	virtual void Resize(float NewSize) = 0;
+};
+
+} // namespace Sequencer
+} // namespace UE
+
diff --git a/Engine/Source/Editor/SequencerCore/Public/MVVM/Extensions/ISoloableExtension.h b/Engine/Source/Editor/SequencerCore/Public/MVVM/Extensions/ISoloableExtension.h
new file mode 100644
--- /dev/null
+++ b/Engine/Source/Editor/SequencerCore/Public/MVVM/Extensions/ISoloableExtension.h
@@ -0,0 +1,30 @@
+// Copyright Epic Games, Inc. All Rights Reserved.
+
+#pragma once
+
+#include "Templates/SharedPointer.h"
+#include "MVVM/ViewModelTypeID.h"
+
+namespace UE
+{
+namespace Sequencer
+{
+
+/**
+ * An extension for outliner nodes that can be made 'solo'
+ */
+class SEQUENCERCORE_API ISoloableExtension
+{
+public:
+
+	UE_SEQUENCER_DECLARE_VIEW_MODEL_TYPE_ID(ISoloableExtension)
+
+	virtual ~ISoloableExtension(){}
+
+	/** Returns whether this item is solo, or in a solo branch */
+	virtual bool IsSolo() const = 0;
+};
+
+} // namespace Sequencer
+} // namespace UE
+
diff --git a/Engine/Source/Editor/SequencerCore/Public/MVVM/Extensions/ITrackLaneExtension.h b/Engine/Source/Editor/SequencerCore/Public/MVVM/Extensions/ITrackLaneExtension.h
new file mode 100644
--- /dev/null
+++ b/Engine/Source/Editor/SequencerCore/Public/MVVM/Extensions/ITrackLaneExtension.h
@@ -0,0 +1,155 @@
+// Copyright Epic Games, Inc. All Rights Reserved.
+
+#pragma once
+
+#include "MVVM/ViewModelTypeID.h"
+#include "MVVM/ViewModels/EditorViewModel.h"
+#include "Widgets/SCompoundWidget.h"
+
+struct FGeometry;
+
+class SWidget;
+class ISequencer;
+class FArrangedWidget;
+struct FTimeToPixel;
+
+namespace UE
+{
+namespace Sequencer
+{
+
+struct FTrackLaneScreenAlignment;
+
+struct SEQUENCERCORE_API FTrackLaneVerticalArrangement
+{
+	float Offset = 0;
+	float Height = 0;
+};
+
+struct SEQUENCERCORE_API FTrackLaneVerticalAlignment
+{
+	enum class ESizeMode : uint8
+	{
+		Proportional, Fixed
+	};
+
+	float              VSizeParam = 1.f;
+	float              VPadding   = 0.f;
+	EVerticalAlignment VAlign     = VAlign_Center;
+	ESizeMode          Mode       = ESizeMode::Proportional;
+
+	FTrackLaneVerticalArrangement ArrangeWithin(float LayoutHeight) const;
+};
+
+struct SEQUENCERCORE_API FTrackLaneVirtualAlignment
+{
+	TRange<FFrameNumber> Range;
+	FTrackLaneVerticalAlignment VerticalAlignment;
+
+	bool IsVisible() const
+	{
+		return !Range.IsEmpty() && !Range.IsDegenerate();
+	}
+
+	static FTrackLaneVirtualAlignment Fixed(const TRange<FFrameNumber>& InRange, float InFixedHeight, EVerticalAlignment InVAlign = VAlign_Center)
+	{
+		return FTrackLaneVirtualAlignment { InRange, { InFixedHeight, 0.f, InVAlign, FTrackLaneVerticalAlignment::ESizeMode::Fixed } };
+	}
+	static FTrackLaneVirtualAlignment Proportional(const TRange<FFrameNumber>& InRange, float InStretchFactor, EVerticalAlignment InVAlign = VAlign_Center)
+	{
+		return FTrackLaneVirtualAlignment { InRange, { InStretchFactor, 0.f, InVAlign, FTrackLaneVerticalAlignment::ESizeMode::Proportional } };
+	}
+
+	TOptional<FFrameNumber> GetFiniteLength() const;
+
+	FTrackLaneScreenAlignment ToScreen(const FTimeToPixel& TimeToPixel, const FGeometry& ParentGeometry) const;
+};
+
+struct SEQUENCERCORE_API FTrackLaneScreenAlignment
+{
+	float LeftPosPx = 0.f;
+	float WidthPx = 0.f;
+	FTrackLaneVerticalAlignment VerticalAlignment;
+
+	bool IsVisible() const
+	{
+		return WidthPx > 0.f;
+	}
+
+	FArrangedWidget ArrangeWidget(TSharedRef<SWidget> InWidget, const FGeometry& ParentGeometry) const;
+};
+
+struct SEQUENCERCORE_API FArrangedVirtualEntity
+{
+	TRange<FFrameNumber> Range;
+	float VirtualTop, VirtualBottom;
+};
+
+/** Base interface for track-area lanes */
+class SEQUENCERCORE_API ITrackLaneWidget
+{
+public:
+
+	/**
+	 * Retrieve this interface as a widget
+	 */
+	virtual TSharedRef<const SWidget> AsWidget() const = 0;
+
+	TSharedRef<SWidget> AsWidget()
+	{
+		TSharedRef<const SWidget> ConstWidget = const_cast<const ITrackLaneWidget*>(this)->AsWidget();
+		return ConstCastSharedRef<SWidget>(ConstWidget);
+	}
+
+	/**
+	 * Arrange this widget within its parent slot
+	 */
+	virtual FTrackLaneScreenAlignment GetAlignment(const FTimeToPixel& TimeToPixel, const FGeometry& InParentGeometry) const = 0;
+
+	/**
+	 * Gets this widget's overlap priority
+	 */
+	virtual int32 GetOverlapPriority() const { return 0; }
+
+	/**
+	 * Receive parent geometry for this lane in Desktop space
+	 */
+	virtual void ReportParentGeometry(const FGeometry& InDesktopSpaceParentGeometry) {}
+
+	/**
+	 * Whether this track lane accepts child widgets
+	 */
+	virtual bool AcceptsChildren() const { return false; }
+
+	/**
+	 * Add a new child to this lane
+	 */
+	virtual void AddChildLane(TSharedPtr<ITrackLaneWidget> ChildWidget) {}
+};
+
+/** Parameters for creating a track lane widget */
+struct SEQUENCERCORE_API FCreateTrackLaneViewParams
+{
+	FCreateTrackLaneViewParams(const TSharedPtr<FEditorViewModel> InEditor)
+		: Editor(InEditor)
+	{}
+
+	const TSharedPtr<FEditorViewModel> Editor;
+};
+
+/** Extension for view-models that can create track lanes in the track area */
+class SEQUENCERCORE_API ITrackLaneExtension
+{
+public:
+
+	UE_SEQUENCER_DECLARE_VIEW_MODEL_TYPE_ID(ITrackLaneExtension)
+
+	virtual ~ITrackLaneExtension(){}
+
+	virtual TSharedPtr<ITrackLaneWidget> CreateTrackLaneView(const FCreateTrackLaneViewParams& InParams) = 0;
+	virtual FTrackLaneVirtualAlignment ArrangeVirtualTrackLaneView() const = 0;
+};
+
+} // namespace Sequencer
+} // namespace UE
+
diff --git a/Engine/Source/Editor/SequencerCore/Public/MVVM/ICastable.h b/Engine/Source/Editor/SequencerCore/Public/MVVM/ICastable.h
new file mode 100644
--- /dev/null
+++ b/Engine/Source/Editor/SequencerCore/Public/MVVM/ICastable.h
@@ -0,0 +1,226 @@
+// Copyright Epic Games, Inc. All Rights Reserved.
+
+#pragma once
+
+#include "Templates/Invoke.h"
+#include "Templates/SharedPointer.h"
+#include "MVVM/ViewModelTypeID.h"
+
+namespace UE
+{
+namespace Sequencer
+{
+
+struct FCastableBoilerplate;
+
+class SEQUENCERCORE_API ICastable
+{
+public:
+
+	virtual ~ICastable(){}
+
+	template<typename RetType, typename InterfaceType, typename ...ArgTypes, typename ...ParamTypes>
+	RetType CastedCall(RetType(InterfaceType::*MemberFunc)(ArgTypes...), ParamTypes&&... InArgs)
+	{
+		return Invoke(MemberFunc, this->CastThisChecked<InterfaceType>(), Forward<ParamTypes>(InArgs)...);
+	}
+
+	template<typename RetType, typename InterfaceType, typename ...ArgTypes, typename ...ParamTypes>
+	RetType CastedCall(RetType(InterfaceType::*MemberFunc)(ArgTypes...), ParamTypes&&... InArgs) const
+	{
+		return Invoke(MemberFunc, this->CastThisChecked<InterfaceType>(), Forward<ParamTypes>(InArgs)...);
+	}
+
+	template<typename T>
+	static T* CastWeakPtr(const TWeakPtr<ICastable>& InWeakPtr)
+	{
+		if (TSharedPtr<ICastable> Pinned = InWeakPtr.Pin())
+		{
+			return Pinned->CastThis<T>();
+		}
+		return nullptr;
+	}
+
+	template<typename T>
+	static TSharedPtr<T> CastWeakPtrShared(const TWeakPtr<ICastable>& InWeakPtr)
+	{
+		if (TSharedPtr<ICastable> Pinned = InWeakPtr.Pin())
+		{
+			return TSharedPtr<T>(Pinned, Pinned->CastThis<T>());
+		}
+		return nullptr;
+	}
+
+	bool IsA(FViewModelTypeID InType) const
+	{
+		return CastRaw(InType) != nullptr;
+	}
+
+	template<typename T>
+	bool IsA() const
+	{
+		return CastThis<T>() != nullptr;
+	}
+
+	template<typename T>
+	bool IsAny() const;
+
+	void* CastRaw(FViewModelTypeID InType);
+	const void* CastRaw(FViewModelTypeID InType) const;
+
+	template<typename T>
+	T* CastThis()
+	{
+		return static_cast<T*>(const_cast<void*>(CastRaw(T::ID)));
+	}
+
+	template<typename T>
+	const T* CastThis() const
+	{
+		return static_cast<const T*>(CastRaw(T::ID));
+	}
+
+	template<typename T>
+	T* CastThisChecked()
+	{
+		T* Result = CastThis<T>();
+		check(Result);
+		return Result;
+	}
+
+	template<typename T>
+	const T* CastThisChecked() const
+	{
+		const T* Result = CastThis<const T>();
+		check(Result);
+		return Result;
+	}
+
+protected:
+
+	// Usually implemented through UE_SEQUENCER_DECLARE_CASTABLE
+	virtual void CastImpl(FViewModelTypeID Type, const void*& OutPtr) const = 0;
+};
+
+template<typename ...T>
+struct TImplements;
+
+#define UE_SEQUENCER_DECLARE_CASTABLE(ThisType, ...)                                             \
+	UE_SEQUENCER_DECLARE_VIEW_MODEL_TYPE_ID(ThisType);                                           \
+	using Implements = TImplements<__VA_ARGS__>;                                                 \
+	virtual void CastImpl(FViewModelTypeID ToType, const void*& OutPtr) const override  \
+	{                                                                                            \
+		FCastableBoilerplate::CastImplementation(this, ToType, OutPtr);                          \
+	}
+
+#define UE_SEQUENCER_DEFINE_CASTABLE(ThisType)                                                   \
+	UE_SEQUENCER_DEFINE_VIEW_MODEL_TYPE_ID(ThisType);
+
+
+template<typename ImplementsType>
+struct TCompositeCast;
+
+struct FCastableBoilerplate
+{
+	template<typename T>
+	struct THasImplements
+	{
+		template<typename U, typename V = typename U::Implements>
+		static uint8 HasImplements(const U* Input);
+		static uint16 HasImplements(const void* Input);
+
+		static constexpr bool Value = sizeof(HasImplements((const T*)0)) == 1;
+	};
+
+	template<typename T>
+	FORCENOINLINE static void CastImplementation(const T* This, FViewModelTypeID ToType, const void*& OutResult)
+	{
+		if (const void* Result = DirectCast(This, ToType))
+		{
+			OutResult = Result;
+		}
+		else
+		{
+			CompositeCast(This, ToType, OutResult);
+		}
+	}
+
+	template<typename T>
+	FORCENOINLINE static bool IsAnyImplementation(const ICastable* This)
+	{
+		const void* Result = This->CastRaw(T::ID);
+		if (!Result)
+		{
+			return CompositeIsAny<T>(This);
+		}
+	}
+
+	template<typename T, typename U = decltype(T::ID)>
+	FORCENOINLINE static const void* DirectCast(const T* Input, FViewModelTypeID ToType)
+	{
+		if (ToType == T::ID)
+		{
+			return Input;
+		}
+		return nullptr;
+	}
+
+	FORCENOINLINE static const void* DirectCast(const void* Input, FViewModelTypeID ToType)
+	{
+		return nullptr;
+	}
+
+	template<typename T>
+	FORCENOINLINE static std::enable_if_t<THasImplements<T>::Value> CompositeCast(const T* Input, FViewModelTypeID ToType, const void*& OutResult)
+	{
+		TCompositeCast<typename T::Implements>::Apply(Input, ToType, OutResult);
+	}
+
+	FORCENOINLINE static void CompositeCast(const void* Input, FViewModelTypeID ToType, const void*& OutResult)
+	{
+	}
+
+	template<typename T>
+	FORCENOINLINE static std::enable_if_t<THasImplements<T>::Value, bool> CompositeIsAny(const ICastable* This)
+	{
+		return TCompositeCast<typename T::Implements>::IsAny(This);
+	}
+
+	template<typename T>
+	FORCENOINLINE static std::enable_if_t<!THasImplements<T>::Value, bool> CompositeIsAny(const ICastable* This)
+	{
+	}
+};
+
+template<typename ...CastableTo>
+struct TCompositeCast<TImplements<CastableTo...>>
+{
+	template<typename T>
+	FORCENOINLINE static void Apply(const T* This, FViewModelTypeID ToType, const void*& OutResult)
+	{
+		(FCastableBoilerplate::CastImplementation(static_cast<const CastableTo*>(This), ToType, OutResult), ...);
+	}
+	template<typename T>
+	FORCENOINLINE static bool IsAny(const ICastable* This)
+	{
+		bool bIsAny = false;
+		(( bIsAny |= FCastableBoilerplate::IsAnyImplementation<CastableTo>(This) ), ...);
+		return bIsAny;
+	}
+};
+
+template<typename ...T>
+struct TImplements
+{
+};
+
+
+template<typename T>
+bool ICastable::IsAny() const
+{
+	return FCastableBoilerplate::template IsAnyImplementation<T>(this);
+}
+
+} // namespace Sequencer
+} // namespace UE
+
diff --git a/Engine/Source/Editor/SequencerCore/Public/MVVM/SharedViewModelData.h b/Engine/Source/Editor/SequencerCore/Public/MVVM/SharedViewModelData.h
new file mode 100644
--- /dev/null
+++ b/Engine/Source/Editor/SequencerCore/Public/MVVM/SharedViewModelData.h
@@ -0,0 +1,42 @@
+// Copyright Epic Games, Inc. All Rights Reserved.
+
+#pragma once
+
+#include "CoreMinimal.h"
+#include "SequencerCoreFwd.h"
+#include "Templates/SharedPointer.h"
+
+namespace UE
+{
+namespace Sequencer
+{
+
+class FViewModel;
+
+class SEQUENCERCORE_API FSharedViewModelData : public TSharedFromThis<FSharedViewModelData>
+{
+public:
+	void PreHierarchicalChange(const TSharedPtr<FViewModel>& InChangedModel);
+	void BroadcastHierarchicalChange(const TSharedPtr<FViewModel>& InChangedModel);
+
+	void ReportLatentHierarchicalOperations();
+
+	FSimpleMulticastDelegate& SubscribeToHierarchyChanged(const TSharedPtr<FViewModel>& InModel);
+	void UnsubscribeFromHierarchyChanged(const TSharedPtr<FViewModel>& InModel, FDelegateHandle InHandle);
+	void UnsubscribeFromHierarchyChanged(const TSharedPtr<FViewModel>& InModel, const void* InUserObject);
+
+	void PurgeStaleHandlers();
+
+private:
+	friend FViewModelHierarchyOperation;
+
+	TMap<TWeakPtr<FViewModel>, FSimpleMulticastDelegate> HierarchyChangedEventsByModel;
+
+	FViewModelHierarchyOperation* CurrentHierarchicalOperation = nullptr;
+	TUniquePtr<FViewModelHierarchyOperation> LatentOperation;
+};
+
+
+} // namespace Sequencer
+} // namespace UE
+
diff --git a/Engine/Source/Editor/SequencerCore/Public/MVVM/ViewModelPtr.h b/Engine/Source/Editor/SequencerCore/Public/MVVM/ViewModelPtr.h
new file mode 100644
--- /dev/null
+++ b/Engine/Source/Editor/SequencerCore/Public/MVVM/ViewModelPtr.h
@@ -0,0 +1,705 @@
+// Copyright Epic Games, Inc. All Rights Reserved.
+
+#pragma once
+
+#include "SequencerCoreFwd.h"
+#include "Templates/SharedPointer.h"
+#include "Misc/GeneratedTypeName.h"
+#include "MVVM/ViewModels/ViewModel.h"
+
+namespace UE
+{
+namespace Sequencer
+{
+
+
+/**
+ * Pointer type that wraps a shared view model pointer to provide implicit dynamic_cast style
+ *    casting to any other view-model or extension. Implicit conversion is implemented through
+ *    a proxy object so as to make the 'implicit' conversion somewhat explicit:
+ *         FViewModelPtr ModelPtr = GetViewModel();
+ *         TSharedPtr<IDraggableOutlinerExtension>   Draggable = ModelPtr.ImplicitCast();
+ *         TSharedPtr<IDroppableExtension>   Droppable = ModelPtr.ImplicitCast();
+ */
+struct FViewModelPtr
+{
+	FViewModelPtr() = default;
+	FViewModelPtr(const FViewModelPtr&) = default;
+	FViewModelPtr& operator=(const FViewModelPtr&) = default;
+	FViewModelPtr(FViewModelPtr&&) = default;
+	FViewModelPtr& operator=(FViewModelPtr&&) = default;
+
+	/** Construction from a nullptr */
+	FViewModelPtr(nullptr_t)
+		: Model(nullptr)
+	{}
+	/** Assignment from a nullptr */
+	FViewModelPtr& operator=(nullptr_t)
+	{
+		Model = nullptr;
+		return *this;
+	}
+
+	/** Construction from a shared ptr */
+	template<typename ModelType>
+	FViewModelPtr(const TSharedPtr<ModelType>& InModel)
+		: Model(InModel)
+	{}
+	/** Assignment from a shared ptr */
+	template<typename ModelType>
+	FViewModelPtr& operator=(const TSharedPtr<ModelType>& InModel)
+	{
+		Model = InModel;
+		return *this;
+	}
+
+	/** Construction from a raw ptr */
+	FViewModelPtr(FViewModel* InModel)
+		: Model(InModel->AsShared())
+	{}
+	/** Assignment from a raw ptr */
+	FViewModelPtr& operator=(FViewModel* InModel)
+	{
+		Model = InModel->AsShared();
+		return *this;
+	}
+
+public:
+
+	/**
+	 * Check whether this ptr is valid (non-null)
+	 */
+	explicit operator bool() const
+	{
+		return Model.IsValid();
+	}
+
+	/**
+	 * Reset this view model pointer back to its default (null) state
+	 */
+	void Reset()
+	{
+		Model = nullptr;
+	}
+
+	/**
+	 * Member access operator giving access to the wrapped view model
+	 */
+	FViewModel* operator->() const
+	{
+		checkf(Model, TEXT("Attempting to access a nullptr FViewModelPtr."));
+		return Model.Get();
+	}
+
+	/**
+	 * Member access operator giving access to the wrapped view model
+	 */
+	FViewModel* Get() const
+	{
+		return Model.Get();
+	}
+
+	/**
+	 * Dereference operator giving access to the wrapped view model
+	 */
+	FViewModel& operator*() const
+	{
+		checkf(Model, TEXT("Attempting to access a nullptr FViewModelPtr."));
+		return *Model.Get();
+	}
+
+	/**
+	 * Implicit conversion to a regular shared pointer
+	 */
+	operator TSharedPtr<FViewModel>() const
+	{
+		return Model;
+	}
+
+	/**
+	 * Explicitly access the view model as a regular shared pointer
+	 */
+	TSharedPtr<FViewModel> AsModel() const
+	{
+		return Model;
+	}
+
+	/**
+	 * Return a proxy type that is able to implicitly convert this view
+	 * model to any other TViewModel<IMyExtension> or TSharedPtr<IMyExtension>
+	 */
+	SEQUENCERCORE_API FImplicitViewModelCast ImplicitCast() const;
+
+	/**
+	 * Type hash function for use with the TSet and TMap family of containers
+	 */
+	friend uint32 GetTypeHash(const FViewModelPtr& In)
+	{
+		return GetTypeHash(In.Model);
+	}
+
+	/**
+	 * Equality operator comparing 2 view model pointers
+	 */
+	friend bool operator==(const FViewModelPtr& A, const FViewModelPtr& B)
+	{
+		return A.Model == B.Model;
+	}
+	friend bool operator!=(const FViewModelPtr& A, const FViewModelPtr& B)
+	{
+		return A.Model != B.Model;
+	}
+
+protected:
+
+	TSharedPtr<FViewModel> Model;
+};
+
+
+/**
+ * Templated view model pointer that enables conversion and access to an interface implemented by a view model ptr.
+ *
+ * Since TViewModelPtr is a templated type, it can also be used as a self-documenting mechanism to specify
+ *   that a parameter should implement a specific extension interface while only passing a single parameter, for example:
+ *   void ExpandRecursively(const TViewModelPtr<IOutlinerExtension>& OutlinerItem);
+ *
+ * Validity is transitive meaning that an invalid TViewModelPtr<A> is equivalent to a nullptr, and will return an
+ * invalid TViewModelPtr<B> even if the originating model implements B but not A.
+ */
+template<typename ExtensionType>
+struct TViewModelPtr : FViewModelPtr
+{
+	TViewModelPtr() = default;
+	TViewModelPtr(const TViewModelPtr<ExtensionType>&) = default;
+	TViewModelPtr<ExtensionType>& operator=(const TViewModelPtr<ExtensionType>&) = default;
+	TViewModelPtr(TViewModelPtr<ExtensionType>&&) = default;
+	TViewModelPtr<ExtensionType>& operator=(TViewModelPtr<ExtensionType>&&) = default;
+
+	/** Construction from a regular view model ptr */
+	TViewModelPtr(const FViewModelPtr& InModelPtr)
+		: Extension(nullptr)
+	{
+		Assign(InModelPtr);
+	}
+	/** Assignment from a regular view model ptr */
+	TViewModelPtr& operator=(const FViewModelPtr& InModelPtr)
+	{
+		Assign(InModelPtr);
+		return *this;
+	}
+
+	/** Construction from a regular view model ptr */
+	TViewModelPtr(const TSharedPtr<FViewModel>& InModel)
+		: Extension(nullptr)
+	{
+		Assign(InModel);
+	}
+	/** Assignment from a regular view model ptr */
+	TViewModelPtr& operator=(const TSharedPtr<FViewModel>& InModel)
+	{
+		Assign(InModel);
+		return *this;
+	}
+
+	/** Construction from a regular view model ref */
+	TViewModelPtr(const TSharedRef<FViewModel>& InModel)
+		: Extension(nullptr)
+	{
+		Assign(InModel);
+	}
+	/** Assignment from a regular view model ref */
+	TViewModelPtr& operator=(const TSharedRef<FViewModel>& InModel)
+	{
+		Assign(InModel);
+		return *this;
+	}
+
+	/** Construction from a regular view model ptr */
+	template<typename SharedType>
+	TViewModelPtr(const TSharedPtr<SharedType>& InModel)
+		: Extension(nullptr)
+	{
+		Assign(InModel);
+	}
+	/** Assignment from a regular view model ptr */
+	template<typename SharedType>
+	TViewModelPtr& operator=(const TSharedPtr<SharedType>& InModel)
+	{
+		Assign(InModel);
+		return *this;
+	}
+
+	/** Construction from a raw ptr */
+	TViewModelPtr(FViewModel* InModel)
+	{
+		if (InModel)
+		{
+			Assign(InModel->AsShared());
+		}
+	}
+	/** Assignment from a raw ptr */
+	TViewModelPtr<ExtensionType>& operator=(FViewModel* InModel)
+	{
+		if (InModel)
+		{
+			Assign(InModel->AsShared());
+		}
+		else
+		{
+			Reset();
+		}
+		return *this;
+	}
+
+	/**
+	 * Implicit construction from another view model ptr of a related type (ie, TViewModelPtr<Base>(TViewModelPtr<Derived>) )
+	 * We leave the compiler to report errors on the conversion of In.Extension -> Extension for unrelated types
+	 */
+	template<typename OtherExtensionType>
+	TViewModelPtr(const TViewModelPtr<OtherExtensionType>& In)
+		: FViewModelPtr(In.Model)
+		, Extension(In.Extension)
+	{}
+	/**
+	 * Implicit assignment from another view model ptr of a related type (ie, TViewModelPtr<Base>(TViewModelPtr<Derived>) )
+	 * We leave the compiler to report errors on the conversion of In.Extension -> Extension for unrelated types
+	 */
+	template<typename OtherExtensionType>
+	TViewModelPtr<ExtensionType>& operator=(const TViewModelPtr<OtherExtensionType>& In)
+	{
+		Model = In.Model;
+		Extension = In.Extension;
+		return *this;
+	}
+
+	/**
+	 * Construction from nullptr
+	 */
+	TViewModelPtr(nullptr_t)
+		: FViewModelPtr(nullptr)
+		, Extension(nullptr)
+	{}
+	/**
+	 * Assignment from nullptr
+	 */
+	TViewModelPtr<ExtensionType>& operator=(nullptr_t)
+	{
+		Model = nullptr;
+		Extension = nullptr;
+		return *this;
+	}
+
+	/**
+	 * Conversion to a normal shared pointer of this extension type
+	 */
+	operator TSharedPtr<ExtensionType>() const
+	{
+		return TSharedPtr<ExtensionType>(Model, Extension);
+	}
+
+public:
+
+	/**
+	 * Check whether this ptr is valid (non-null)
+	 * @note: While a TViewModelPtr may contain a valid model ptr, it is only
+	 *        considered valid as long as the model implements the templated extension or model type
+	 */
+	explicit operator bool() const
+	{
+		return Extension != nullptr;
+	}
+
+	/**
+	 * Pointer access to the extension pointer. This Ptr must be checked for validity before calling.
+	 */
+	ExtensionType* operator->() const
+	{
+		checkf(Extension, TEXT("Attempting to access a nullptr TViewModelPtr<%s>."), GetGeneratedTypeName<ExtensionType>());
+		return Extension;
+	}
+
+	/**
+	 * Dereference operator giving access to the wrapped extension type
+	 */
+	ExtensionType& operator*() const
+	{
+		checkf(Extension, TEXT("Attempting to access a nullptr FViewModelPtr."));
+		return *Extension;
+	}
+
+	/**
+	 * Conversion to a normal shared pointer of this extension type
+	 */
+	ExtensionType* Get() const
+	{
+		return Extension;
+	}
+
+	/**
+	 * Reset this ptr back to its default (null) state
+	 */
+	void Reset()
+	{
+		Model = nullptr;
+		Extension = nullptr;
+	}
+
+	friend bool operator==(const TViewModelPtr<ExtensionType>& A, const TViewModelPtr<ExtensionType>& B)
+	{
+		return A.Model == B.Model && A.Extension == B.Extension;
+	}
+
+private:
+
+	void Assign(const TSharedPtr<FViewModel>& InModel, ExtensionType* InExtension = nullptr)
+	{
+		if (InExtension)
+		{
+			Extension = InExtension;
+			Model = InModel;
+			return;
+		}
+
+		if (InModel)
+		{
+			if (ExtensionType* NewExtension = InModel->CastThis<ExtensionType>())
+			{
+				Extension = NewExtension;
+				Model = InModel;
+				return;
+			}
+		}
+
+		Extension = nullptr;
+		Model = nullptr;
+	}
+
+	template<typename> friend struct TViewModelPtr;
+
+	ExtensionType* Extension = nullptr;
+};
+
+/**
+ * Weak pointer utility for view models that also provides the same explicit conversion utility as F/TViewModelPtr
+ * Using this pointer type allows for expressive null handling and casting in a simple function call for any
+ * extension type which is useful for many widgets:
+ *
+ *    class SMyWidget
+ *    {
+ *        void OnMouseMove()
+ *        {
+ *            TViewModelPtr<IDraggable> Draggable = WeakViewModel.ImplicitPin();
+ *            if (Draggable)
+ *            {
+ *                Draggable->BeginDrag();
+ *            }
+ *        }
+ *        FWeakViewModelPtr WeakViewModel;
+ *    };
+ */
+struct FWeakViewModelPtr
+{
+	FWeakViewModelPtr() = default;
+
+	/**
+	 * Construction from a shared ref to a view model
+	 * This is templated to ensure that the overloaded constructors remain unambiguous
+	 */
+	template<typename SharedType>
+	FWeakViewModelPtr(const TSharedRef<SharedType>& InModel)
+		: WeakModel(InModel)
+	{}
+
+	/**
+	 * Construction from a shared ptr to a view model
+	 * This is templated to ensure that the overloaded constructors remain unambiguous
+	 */
+	template<typename SharedType>
+	FWeakViewModelPtr(const TSharedPtr<SharedType>& InModel)
+		: WeakModel(InModel)
+	{}
+
+	/**
+	 * Construction from a weak ptr to a view model
+	 * This is templated to ensure that the overloaded constructors remain unambiguous
+	 */
+	template<typename SharedType>
+	FWeakViewModelPtr(const TWeakPtr<SharedType>& InWeakModel)
+		: WeakModel(InWeakModel)
+	{}
+
+	/**
+	 * Construction from a view model ptr
+	 */
+	FWeakViewModelPtr(const FViewModelPtr& InViewModelPtr)
+		: WeakModel(InViewModelPtr.AsModel())
+	{}
+
+	/**
+	 * Assignment from a view model ptr
+	 */
+	FWeakViewModelPtr& operator=(const FViewModelPtr& InViewModelPtr)
+	{
+		WeakModel = InViewModelPtr.AsModel();
+		return *this;
+	}
+
+	/**
+	 * Pin this weak ptr. The result is guaranteed to be kept alive as long as the variable's scope.
+	 */
+	FViewModelPtr Pin() const
+	{
+		TSharedPtr<FViewModel> Pinned = WeakModel.Pin();
+		return FViewModelPtr(Pinned);
+	}
+
+	/**
+	 * Returns whether the weak ptr is valid.
+	 */
+	bool IsValid() const
+	{
+		return WeakModel.IsValid();
+	}
+
+	/**
+	 * Conversion to a normal weak pointer
+	 */
+	operator TWeakPtr<FViewModel>() const
+	{
+		return WeakModel;
+	}
+
+	/**
+	 * Attempt to convert this weak view model ptr to another type implicitly
+	 */
+	SEQUENCERCORE_API FImplicitWeakViewModelCast ImplicitCast() const;
+
+	/**
+	 * Returns a proxy object that will simultaneously Pin() and Cast() this view model
+	 * to a particular view-model or extension type
+	 */
+	SEQUENCERCORE_API FImplicitWeakViewModelPin ImplicitPin() const;
+
+	TWeakPtr<FViewModel> AsModel() const
+	{
+		return WeakModel;
+	}
+
+	friend uint32 GetTypeHash(const FWeakViewModelPtr& In)
+	{
+		return GetTypeHash(In.WeakModel);
+	}
+
+	friend bool operator==(const FWeakViewModelPtr& A, const FWeakViewModelPtr& B)
+	{
+		return A.WeakModel == B.WeakModel;
+	}
+
+	friend bool operator!=(const FWeakViewModelPtr& A, const FWeakViewModelPtr& B)
+	{
+		return A.WeakModel != B.WeakModel;
+	}
+protected:
+
+	TWeakPtr<FViewModel> WeakModel;
+};
+
+/**
+ * Templated version of FWeakViewModelPtr that provides an additional self-documenting extension that should
+ * exist on the view model.
+ */
+template<typename ExtensionType>
+struct TWeakViewModelPtr : FWeakViewModelPtr
+{
+	TWeakViewModelPtr() = default;
+	TWeakViewModelPtr(const TWeakViewModelPtr<ExtensionType>&) = default;
+	TWeakViewModelPtr<ExtensionType>& operator=(const TWeakViewModelPtr<ExtensionType>&) = default;
+	TWeakViewModelPtr(TWeakViewModelPtr<ExtensionType>&&) = default;
+	TWeakViewModelPtr<ExtensionType>& operator=(TWeakViewModelPtr<ExtensionType>&&) = default;
+
+	/**
+	 * Construction from a shared ref to a view model
+	 * This is templated to ensure that the overloaded constructors remain unambiguous
+	 */
+	template<typename SharedType>
+	TWeakViewModelPtr(const TSharedRef<SharedType>& InModel)
+		: FWeakViewModelPtr(InModel)
+	{}
+
+	/**
+	 * Construction from a shared ptr to a view model
+	 * This is templated to ensure that the overloaded constructors remain unambiguous
+	 */
+	template<typename SharedType>
+	TWeakViewModelPtr(const TSharedPtr<SharedType>& InModel)
+		: FWeakViewModelPtr(InModel)
+	{}
+
+	/**
+	 * Construction from a weak ptr to a view model
+	 * This is templated to ensure that the overloaded constructors remain unambiguous
+	 */
+	template<typename SharedType>
+	TWeakViewModelPtr(const TWeakPtr<SharedType>& InWeakModel)
+		: FWeakViewModelPtr(InWeakModel)
+	{}
+
+	/** Construction from a FWeakViewModelPtr */
+	TWeakViewModelPtr(const FWeakViewModelPtr& InModel)
+		: FWeakViewModelPtr(InModel)
+	{}
+	/** Assignment from a FWeakViewModelPtr */
+	TWeakViewModelPtr<ExtensionType>& operator=(const FWeakViewModelPtr& InModel)
+	{
+		WeakModel = InModel.AsModel();
+		return *this;
+	}
+
+	/** Construction from a TViewModelPtr */
+	TWeakViewModelPtr(const TViewModelPtr<ExtensionType>& InModel)
+		: FWeakViewModelPtr(InModel.AsModel())
+	{}
+	/** Assignment from a TViewModelPtr */
+	TWeakViewModelPtr<ExtensionType>& operator=(const TViewModelPtr<ExtensionType>& InModel)
+	{
+		WeakModel = InModel.AsModel();
+		return *this;
+	}
+
+	/** Pin this weak view model ptr */
+	TViewModelPtr<ExtensionType> Pin() const
+	{
+		TSharedPtr<FViewModel> Pinned = WeakModel.Pin();
+		return TViewModelPtr<ExtensionType>(Pinned);
+	}
+
+	/**
+	 * Conversion to a normal weak pointer
+	 */
+	operator TWeakPtr<ExtensionType>() const
+	{
+		TViewModelPtr<ExtensionType> Pinned = Pin();
+		return TSharedPtr<ExtensionType>(Pinned);
+	}
+};
+
+
+/**
+ * Proxy type that enables implicit dynamic casting of a view-model to any other view-model
+ * or extension type. This is implemented as a proxy type rather than the conversion
+ * operators existing on FViewModelPtr so that the implicit cast must be user-directed, while
+ * still removing additional template boilerplate at the call-site.
+ *
+ * @see FViewModelPtr::ImplicitCast
+ */
+struct FImplicitViewModelCast
+{
+	/** Construction from a regular shared ptr to a view model */
+	explicit FImplicitViewModelCast(const TSharedPtr<FViewModel>& In)
+		: ViewModel(In)
+	{}
+
+	/**
+	 * Implicitly convert this view model to any other view-model or extension type
+	 * @return A regular shared pointer to the view-model or extension if it is implemented on the existing model; nullptr otherwise.
+	 */
+	template<typename T>
+	operator TSharedPtr<T>() const
+	{
+		return ViewModel ? ViewModel->CastThisShared<T>() : nullptr;
+	}
+
+	/**
+	 * Implicitly convert this view model to any other view-model or extension type
+	 * @return A TViewModelPtr to the view-model or extension if it is implemented on the existing model; nullptr otherwise.
+	 */
+	template<typename ExtensionType>
+	operator TViewModelPtr<ExtensionType>() const
+	{
+		return TViewModelPtr<ExtensionType>(ViewModel);
+	}
+
+private:
+
+	TSharedPtr<FViewModel> ViewModel;
+};
+
+
+
+/**
+ * Proxy type that enables implicit dynamic casting of a weak view-model to any other view-model
+ * or extension type. This is implemented as a proxy type rather than the conversion
+ * operators existing on FWeakViewModelPtr so that the implicit cast must be user-directed, while
+ * still removing additional template boilerplate at the call-site.
+ *
+ * @see FWeakViewModelPtr::ImplicitCast
+ * @see FWeakViewModelPtr::ImplicitPin
+ */
+struct FImplicitWeakViewModelCast
+{
+	/** Construction from a regular weak ptr to a view model */
+	explicit FImplicitWeakViewModelCast(const TWeakPtr<FViewModel>& In)
+		: WeakModel(In)
+	{}
+
+	/**
+	 * Implicit conversion to any TWeakViewModelPtr type
+	 */
+	template<typename ExtensionType>
+	operator TWeakViewModelPtr<ExtensionType>() const
+	{
+		return TWeakViewModelPtr<ExtensionType>(WeakModel);
+	}
+
+private:
+	TWeakPtr<FViewModel> WeakModel;
+};
+
+
+/**
+ * Proxy type that combines a Pin() and Cast() of a weak view model ptr, gracefully handling null.
+ *
+ * @see FWeakViewModelPtr::ImplicitPin
+ */
+struct FImplicitWeakViewModelPin
+{
+	/** Construction from a regular weak ptr to a view model */
+	explicit FImplicitWeakViewModelPin(const TWeakPtr<FViewModel>& In)
+		: WeakModel(In)
+	{}
+
+	/**
+	 * Simultaneously Pin() and Cast() this view model to the templated type
+	 * @return A shared pointer to the view-model or extension type if it is implemented on this model, nullptr otherwise.
+	 */
+	template<typename T>
+	operator TSharedPtr<T>() const
+	{
+		TSharedPtr<FViewModel> Pinned = WeakModel.Pin();
+		return Pinned ? Pinned->CastThisShared<T>() : nullptr;
+	}
+
+	/**
+	 * Simultaneously Pin() and Cast() this view model to the templated type
+	 * @return A TViewModelPtr to the view-model or extension type if it is implemented on this model, nullptr otherwise.
+	 */
+	template<typename ExtensionType>
+	operator TViewModelPtr<ExtensionType>() const
+	{
+		TSharedPtr<FViewModel> Pinned = WeakModel.Pin();
+		if (Pinned)
+		{
+			return TViewModelPtr<ExtensionType>(Pinned);
+		}
+		return nullptr;
+	}
+
+private:
+
+	TWeakPtr<FViewModel> WeakModel;
+};
+
+
+} // namespace Sequencer
+} // namespace UE
+
diff --git a/Engine/Source/Editor/SequencerCore/Public/MVVM/ViewModelTypeID.h b/Engine/Source/Editor/SequencerCore/Public/MVVM/ViewModelTypeID.h
new file mode 100644
--- /dev/null
+++ b/Engine/Source/Editor/SequencerCore/Public/MVVM/ViewModelTypeID.h
@@ -0,0 +1,79 @@
+// Copyright Epic Games, Inc. All Rights Reserved.
+
+#pragma once
+
+#include "CoreTypes.h"
+
+#define UE_SEQUENCER_DECLARE_VIEW_MODEL_TYPE_ID(Type) static ::UE::Sequencer::TAutoRegisterViewModelTypeID<Type> ID;
+#define UE_SEQUENCER_DEFINE_VIEW_MODEL_TYPE_ID(Type) ::UE::Sequencer::TAutoRegisterViewModelTypeID<Type> Type::ID;
+
+namespace UE
+{
+namespace Sequencer
+{
+
+struct FViewModelTypeID
+{
+	friend uint32 GetTypeHash(FViewModelTypeID In)
+	{
+		return In.ID;
+	}
+
+	friend bool operator<(FViewModelTypeID A, FViewModelTypeID B)
+	{
+		return A.ID < B.ID;
+	}
+
+	friend bool operator==(FViewModelTypeID A, FViewModelTypeID B)
+	{
+		return A.ID == B.ID;
+	}
+
+protected:
+
+	SEQUENCERCORE_API static FViewModelTypeID RegisterNew();
+	SEQUENCERCORE_API static FViewModelTypeID Invalid();
+
+private:
+
+	FViewModelTypeID(uint32 InID)
+		: ID(InID)
+	{}
+
+	FViewModelTypeID()
+		: ID(~0u)
+	{}
+
+	uint32 ID;
+};
+
+template<typename T>
+struct TViewModelTypeID : FViewModelTypeID
+{
+	static TViewModelTypeID<T> RegisterNew()
+	{
+		return TViewModelTypeID<T>(FViewModelTypeID::RegisterNew());
+	}
+
+	static TViewModelTypeID<T> Invalid()
+	{
+		return TViewModelTypeID<T>(FViewModelTypeID::Invalid());
+	}
+
+protected:
+	explicit TViewModelTypeID(FViewModelTypeID In)
+		: FViewModelTypeID(In)
+	{}
+};
+
+template<typename T>
+struct TAutoRegisterViewModelTypeID : TViewModelTypeID<T>
+{
+	TAutoRegisterViewModelTypeID()
+		: TViewModelTypeID<T>(TViewModelTypeID<T>::RegisterNew())
+	{}
+};
+
+} // namespace Sequencer
+} // namespace UE
+
