Sequencer: MVVM2 branch and Layer Bars

diff --git a/Engine/Source/Editor/SequencerCore/Public/MVVM/ViewModels/OutlinerSpacer.h b/Engine/Source/Editor/SequencerCore/Public/MVVM/ViewModels/OutlinerSpacer.h
new file mode 100644
--- /dev/null
+++ b/Engine/Source/Editor/SequencerCore/Public/MVVM/ViewModels/OutlinerSpacer.h
@@ -0,0 +1,60 @@
+// Copyright Epic Games, Inc. All Rights Reserved.
+
+#pragma once
+
+#include "MVVM/ViewModels/ViewModel.h"
+#include "MVVM/Extensions/IGeometryExtension.h"
+#include "MVVM/Extensions/IOutlinerExtension.h"
+#include "MVVM/Extensions/ISortableExtension.h"
+
+namespace UE
+{
+namespace Sequencer
+{
+
+class SEQUENCERCORE_API FOutlinerSpacer
+	: public FViewModel
+	, public FGeometryExtensionShim
+	, public FOutlinerExtensionShim
+	, public ISortableExtension
+	, public IOutlinerDropTargetOutlinerExtension
+{
+public:
+
+	UE_SEQUENCER_DECLARE_CASTABLE(FOutlinerSpacer, FViewModel, IGeometryExtension, IOutlinerExtension, ISortableExtension, IOutlinerDropTargetOutlinerExtension);
+
+	FOutlinerSpacer(float InDesiredSpacerHeight);
+	~FOutlinerSpacer();
+
+	void SetDesiredHeight(float InDesiredSpacerHeight)
+	{
+		DesiredSpacerHeight = InDesiredSpacerHeight;
+	}
+
+public:
+
+	/*~ IOutlinerExtension */
+	bool HasBackground() const override;
+	FName GetIdentifier() const override;
+	FOutlinerSizing GetOutlinerSizing() const override;
+	TSharedRef<SWidget> CreateOutlinerView(const FCreateOutlinerViewParams& InParams) override;
+	TSharedPtr<SWidget> CreateContextMenuWidget(const FCreateOutlinerContextMenuWidgetParams& InParams) override;
+
+	/*~ ISortableExtension */
+	virtual void SortChildren() override;
+	virtual FSortingKey GetSortingKey() const override;
+	virtual void SetCustomOrder(int32 InCustomOrder) override;
+
+	/*~ IOutlinerDropTargetOutlinerExtension */
+	TOptional<EItemDropZone> CanAcceptDrop(const FViewModelPtr& TargetModel, const FDragDropEvent& DragDropEvent, EItemDropZone InItemDropZone) override;
+	void PerformDrop(const FViewModelPtr& TargetModel, const FDragDropEvent& DragDropEvent, EItemDropZone InItemDropZone) override;
+
+private:
+
+	float DesiredSpacerHeight;
+	int32 CustomOrder;
+};
+
+} // namespace Sequencer
+} // namespace UE
+
diff --git a/Engine/Source/Editor/SequencerCore/Public/MVVM/ViewModels/OutlinerViewModel.h b/Engine/Source/Editor/SequencerCore/Public/MVVM/ViewModels/OutlinerViewModel.h
new file mode 100644
--- /dev/null
+++ b/Engine/Source/Editor/SequencerCore/Public/MVVM/ViewModels/OutlinerViewModel.h
@@ -0,0 +1,109 @@
+// Copyright Epic Games, Inc. All Rights Reserved.
+
+#pragma once
+
+#include "CoreTypes.h"
+#include "Input/DragAndDrop.h"
+#include "Templates/SharedPointer.h"
+#include "Delegates/Delegate.h"
+#include "MVVM/ViewModelPtr.h"
+#include "MVVM/ViewModelTypeID.h"
+#include "MVVM/ViewModels/ViewModel.h"
+#include "MVVM/ViewModels/ViewModel.h"
+
+class SWidget;
+
+namespace UE
+{
+namespace Sequencer
+{
+
+class IOutlinerExtension;
+class FViewModel;
+class FEditorViewModel;
+
+class SEQUENCERCORE_API FOutlinerViewModel
+	: public FViewModel
+{
+public:
+
+	UE_SEQUENCER_DECLARE_CASTABLE(FOutlinerViewModel, FViewModel);
+
+	FSimpleMulticastDelegate OnRefreshed;
+
+	/**
+	 * Constructor
+	 */
+	FOutlinerViewModel();
+
+	/**
+	 * Destructor
+	 */
+	virtual ~FOutlinerViewModel() {}
+
+	/**
+	 * Initialize the view model with a given root model
+	 */
+	void Initialize(const FWeakViewModelPtr& InWeakRootDataModel);
+
+	/**
+	 * Get the editor view-model provided by the creation context
+	 */
+	TSharedPtr<FEditorViewModel> GetEditor() const;
+
+	/**
+	 * Get the logical root item of the outliner hierarchy
+	 */
+	FViewModelPtr GetRootItem() const;
+
+	/**
+	 * Retrieve a flat list of the root items in this outliner model
+	 */
+	TArrayView<const TWeakViewModelPtr<IOutlinerExtension>> GetTopLevelItems() const;
+
+	/**
+	 * Set the single hovered model in the tree
+	 */
+	void SetHoveredItem(TViewModelPtr<IOutlinerExtension> InHoveredModel);
+
+	/**
+	 * Get the single hovered model in the tree, possibly nullptr
+	 */
+	TViewModelPtr<IOutlinerExtension> GetHoveredItem() const;
+
+	/**
+	 * Unpins any pinned nodes in this tree
+	 */
+	void UnpinAllNodes();
+
+	/**
+	 * Called when the default context menu is opened
+	 */
+	virtual TSharedPtr<SWidget> CreateContextMenuWidget();
+
+	/**
+	 * Initate a new drag operation from the specified dragged models.
+	 */
+	virtual TSharedRef<FDragDropOperation> InitiateDrag(TArray<TWeakViewModelPtr<IOutlinerExtension>>&& InDraggedModels);
+
+private:
+
+	void HandleDataHierarchyChanged();
+
+private:
+
+	/** The root data */
+	FWeakViewModelPtr WeakRootDataModel;
+	/** A cached list of the root data's children, which are the outliner's root items */
+	TArray<TWeakViewModelPtr<IOutlinerExtension>> RootItems;
+
+	/** The currently mouse-hovered item */
+	TWeakViewModelPtr<IOutlinerExtension> WeakHoveredItem;
+
+	/** Whether root items should be re-evaluated */
+	bool bRootItemsInvalidated;
+};
+
+} // namespace Sequencer
+} // namespace UE
+
diff --git a/Engine/Source/Editor/SequencerCore/Public/MVVM/ViewModels/TrackAreaViewModel.h b/Engine/Source/Editor/SequencerCore/Public/MVVM/ViewModels/TrackAreaViewModel.h
new file mode 100644
--- /dev/null
+++ b/Engine/Source/Editor/SequencerCore/Public/MVVM/ViewModels/TrackAreaViewModel.h
@@ -0,0 +1,91 @@
+// Copyright Epic Games, Inc. All Rights Reserved.
+
+#pragma once
+
+#include "CoreTypes.h"
+#include "MVVM/ViewModels/ViewModel.h"
+#include "Math/Range.h"
+#include "TimeToPixel.h"
+
+struct FFrameRate;
+struct FFrameNumber;
+struct FGeometry;
+
+namespace UE
+{
+namespace Sequencer
+{
+
+struct ITrackAreaHotspot;
+class ISequencerEditTool;
+class FEditorViewModel;
+
+class SEQUENCERCORE_API FTrackAreaViewModel
+	: public FViewModel
+{
+public:
+
+	UE_SEQUENCER_DECLARE_CASTABLE(FTrackAreaViewModel, FViewModel);
+
+	FTrackAreaViewModel();
+	virtual ~FTrackAreaViewModel();
+
+public:
+
+	/** Get the editor view-model provided by the creation context */
+	TSharedPtr<FEditorViewModel> GetEditor() const;
+
+	/** Create a time/pixel converter for the given geometry */
+	FTimeToPixel GetTimeToPixel(const FGeometry& AllottedGeometry) const;
+
+	/** Gets the current tick resolution of the editor */
+	virtual FFrameRate GetTickResolution() const;
+	/** Gets the current view range of this track area */
+	virtual TRange<double> GetViewRange() const;
+
+	/** Get the current active hotspot */
+	TSharedPtr<ITrackAreaHotspot> GetHotspot() const;
+
+	/** Set the hotspot to something else */
+	void SetHotspot(TSharedPtr<ITrackAreaHotspot> NewHotspot);
+	void AddHotspot(TSharedPtr<ITrackAreaHotspot> NewHotspot);
+	void RemoveHotspot(FViewModelTypeID Type);
+	void ClearHotspots();
+
+	/** Set whether the hotspot is locked and cannot be changed (ie when a menu is open) */
+	void LockHotspot(bool bIsLocked);
+
+	void AddEditTool(TSharedPtr<ISequencerEditTool> InNewTool);
+
+	/** Access the currently active track area edit tool */
+	ISequencerEditTool* GetEditTool() const { return EditTool.IsValid() ? EditTool.Get() : nullptr; }
+
+	/** Check whether it's possible to activate the specified tool */
+	bool CanActivateEditTool(FName Identifier) const;
+
+	bool AttemptToActivateTool(FName Identifier);
+
+	void LockEditTool();
+	void UnlockEditTool();
+
+protected:
+
+	/** The current hotspot that can be set from anywhere to initiate drags */
+	TArray<TSharedPtr<ITrackAreaHotspot>> HotspotStack;
+
+	/** The currently active edit tools on this track area */
+	TArray<TSharedPtr<ISequencerEditTool>> EditTools;
+
+	/** The currently active edit tool on this track area */
+	TSharedPtr<ISequencerEditTool> EditTool;
+
+	/** When true, prevents any other hotspot being activated */
+	bool bHotspotLocked;
+
+	/** When true, prevents any other edit tool being activated by way of CanActivateEditTool always returning true */
+	bool bEditToolLocked;
+};
+
+} // namespace Sequencer
+} // namespace UE
+
diff --git a/Engine/Source/Editor/SequencerCore/Public/MVVM/Views/IOutlinerSelectionHandler.h b/Engine/Source/Editor/SequencerCore/Public/MVVM/Views/IOutlinerSelectionHandler.h
new file mode 100644
--- /dev/null
+++ b/Engine/Source/Editor/SequencerCore/Public/MVVM/Views/IOutlinerSelectionHandler.h
@@ -0,0 +1,28 @@
+// Copyright Epic Games, Inc. All Rights Reserved.
+
+#pragma once
+
+#include "CoreTypes.h"
+#include "MVVM/ViewModels/ViewModel.h"
+#include "MVVM/Extensions/IOutlinerExtension.h"
+
+namespace UE
+{
+namespace Sequencer
+{
+
+class IOutlinerExtension;
+
+class SEQUENCERCORE_API IOutlinerSelectionHandler
+	: public TSharedFromThis<IOutlinerSelectionHandler>
+{
+public:
+
+	virtual ~IOutlinerSelectionHandler(){}
+
+	virtual void SelectOutlinerItems(const TArrayView<TWeakViewModelPtr<IOutlinerExtension>>& Items, bool bRightMouseButtonDown) = 0;
+};
+
+} // namespace Sequencer
+} // namespace UE
+
diff --git a/Engine/Source/Editor/SequencerCore/Public/MVVM/Views/SOutlinerView.h b/Engine/Source/Editor/SequencerCore/Public/MVVM/Views/SOutlinerView.h
new file mode 100644
--- /dev/null
+++ b/Engine/Source/Editor/SequencerCore/Public/MVVM/Views/SOutlinerView.h
@@ -0,0 +1,337 @@
+// Copyright Epic Games, Inc. All Rights Reserved.
+
+#pragma once
+
+#include "CoreTypes.h"
+#include "Misc/Attribute.h"
+#include "Input/Reply.h"
+#include "Widgets/SWidget.h"
+#include "Widgets/DeclarativeSyntaxSupport.h"
+#include "Widgets/Views/STableViewBase.h"
+#include "Widgets/Views/STableRow.h"
+#include "Widgets/Views/STreeView.h"
+#include "MVVM/ViewModelPtr.h"
+#include "MVVM/Extensions/IOutlinerExtension.h"
+#include "MVVM/Views/TreeViewTraits.h"
+
+namespace UE
+{
+namespace Sequencer
+{
+
+class IOutlinerSelectionHandler;
+class FViewModel;
+class SOutlinerViewRow;
+class STrackAreaView;
+class STrackLane;
+
+enum class ETreeRecursion
+{
+	Recursive, NonRecursive
+};
+
+/** Structure used to define a column in the tree view */
+struct SEQUENCERCORE_API FOutlinerViewColumn
+{
+	typedef TFunction<TSharedRef<SWidget>(TViewModelPtr<IOutlinerExtension>, const TSharedRef<SOutlinerViewRow>&)> FOnGenerate;
+
+	FOutlinerViewColumn(const FOnGenerate& InOnGenerate, const TAttribute<float>& InWidth) : Generator(InOnGenerate), Width(InWidth) {}
+	FOutlinerViewColumn(FOnGenerate&& InOnGenerate, const TAttribute<float>& InWidth) : Generator(MoveTemp(InOnGenerate)), Width(InWidth) {}
+
+	/** Function used to generate a cell for this column */
+	FOnGenerate Generator;
+	/** Attribute specifying the width of this column */
+	TAttribute<float> Width;
+};
+
+
+/** The tree view used in the sequencer */
+class SEQUENCERCORE_API SOutlinerView 
+	: public STreeView<TWeakViewModelPtr<IOutlinerExtension>>
+{
+public:
+
+	SLATE_BEGIN_ARGS(SOutlinerView){}
+
+		SLATE_ATTRIBUTE( TSharedPtr<IOutlinerSelectionHandler>, SelectionHandler )
+		/** Externally supplied scroll bar */
+		SLATE_ARGUMENT( TSharedPtr<SScrollBar>, ExternalScrollbar )
+
+	SLATE_END_ARGS()
+
+	static const FName TrackNameColumn;
+
+	~SOutlinerView();
+
+	/** Construct this widget */
+	void Construct(const FArguments& InArgs, TWeakPtr<FOutlinerViewModel> InWeakOutliner, const TSharedRef<STrackAreaView>& InTrackArea);
+	virtual void Tick( const FGeometry& AllottedGeometry, const double InCurrentTime, const float InDeltaTime );
+	virtual int32 OnPaint(const FPaintArgs& Args, const FGeometry& AllottedGeometry, const FSlateRect& MyCullingRect, FSlateWindowElementList& OutDrawElements, int32 LayerId, const FWidgetStyle& InWidgetStyle, bool bParentEnabled) const override;
+
+	TSharedPtr<FOutlinerViewModel> GetOutlinerModel() const;
+
+	/** @return the number of root nodes this tree contains */
+	int32 GetNumRootNodes() const { return RootNodes.Num(); }
+
+	float GetVirtualTop() const;
+
+	void GetVisibleItems(TArray<TViewModelPtr<IOutlinerExtension>>& OutItems) const;
+
+	void ForceSetSelectedItems(const TSet<TWeakPtr<FViewModel>>& InItems);
+
+public:
+
+	/** Get the tree item model at the specified physical vertical position */
+	TViewModelPtr<IOutlinerExtension> HitTestNode(float InPhysical) const;
+
+	/** Convert the specified physical vertical position into an absolute virtual position, ignoring expanded states */
+	float PhysicalToVirtual(float InPhysical) const;
+
+	/** Convert the specified absolute virtual position into a physical position in the tree.
+	 * @note: Will not work reliably for virtual positions that are outside of the physical space
+	 */
+	float VirtualToPhysical(float InVirtual) const;
+
+	void ReportChildRowGeometry(const TViewModelPtr<IOutlinerExtension>& InNode, const FGeometry& InGeometry);
+
+public:
+
+	/** Refresh this tree as a result of the underlying tree data changing */
+	void Refresh();
+
+	/** Expand or collapse nodes */
+	void ToggleExpandCollapseNodes(ETreeRecursion Recursion = ETreeRecursion::Recursive, bool bExpandAll = false, bool bCollapseAll = false);
+
+	/** Scroll this tree view by the specified number of slate units */
+	void ScrollByDelta(float DeltaInSlateUnits);
+
+protected:
+
+	/** Set the item's expansion state, including all of its children */
+	void ExpandCollapseNode(TViewModelPtr<IOutlinerExtension> InDataModel, bool bExpansionState, ETreeRecursion Recursion);
+
+	/** Generate a row for a particular node */
+	TSharedRef<ITableRow> OnGenerateRow(TWeakViewModelPtr<IOutlinerExtension> InDisplayNode, const TSharedRef<STableViewBase>& OwnerTable);
+
+	void CreateTrackLanesForRow(TSharedRef<SOutlinerViewRow> InRow, TViewModelPtr<IOutlinerExtension> InDataModel);
+	TSharedPtr<STrackLane> FindOrCreateParentLane(TViewModelPtr<IOutlinerExtension> InDataModel);
+
+	/** Gather the children from the specified node */
+	void OnGetChildren(TWeakViewModelPtr<IOutlinerExtension> InParent, TArray<TWeakViewModelPtr<IOutlinerExtension>>& OutChildren) const;
+
+	/** Generate a widget for the specified Node and Column */
+	TSharedRef<SWidget> GenerateWidgetForColumn(TViewModelPtr<IOutlinerExtension> InDataModel, const FName& ColumnId, const TSharedRef<SOutlinerViewRow>& Row) const;
+
+	/** Called when a node has been expanded or collapsed */
+	void OnExpansionChanged(TWeakViewModelPtr<IOutlinerExtension> InItem, bool bIsExpanded);
+
+	// Tree selection methods which must be overriden to maintain selection consistency with the rest of sequencer.
+	virtual void Private_SetItemSelection( TWeakViewModelPtr<IOutlinerExtension> TheItem, bool bShouldBeSelected, bool bWasUserDirected = false ) override;
+	virtual void Private_ClearSelection() override;
+	virtual void Private_SelectRangeFromCurrentTo( TWeakViewModelPtr<IOutlinerExtension> InRangeSelectionEnd ) override;
+	virtual void Private_SignalSelectionChanged( ESelectInfo::Type SelectInfo ) override;
+
+	virtual void OnRightMouseButtonDown(const FPointerEvent& MouseEvent) override;
+	virtual void OnRightMouseButtonUp(const FPointerEvent& MouseEvent) override;
+
+	virtual FReply OnDragOver(const FGeometry& MyGeometry, const FDragDropEvent& DragDropEvent) override;
+	virtual FReply OnDrop(const FGeometry& MyGeometry, const FDragDropEvent& DragDropEvent) override;
+
+public:
+
+	virtual FReply OnKeyDown(const FGeometry& MyGeometry, const FKeyEvent& InKeyEvent) override;
+
+private:
+
+	// Private, unimplemented overloaded name for SetItemSelection to prevent external calls - use ForceSetItemSelection instead
+	void SetItemSelection();
+	// Private, unimplemented overloaded name for SetItemSelection to prevent external calls - use ForceSetItemSelection instead
+	void ClearSelection();
+
+	void HandleTableViewScrolled(double InScrollOffset);
+	void UpdatePhysicalGeometry(bool bIsRefresh);
+
+	/** Handles the context menu opening when right clicking on the tree view. */
+	TSharedPtr<SWidget> OnContextMenuOpening();
+
+	void SetItemExpansionRecursive(TWeakViewModelPtr<IOutlinerExtension> InItem, bool bIsExpanded);
+
+public:
+
+	/** Structure used to cache physical geometry for a particular node */
+	struct FCachedGeometry
+	{
+		FCachedGeometry(TWeakViewModelPtr<IOutlinerExtension> InWeakItem,
+			float InPhysicalTop, float InPhysicalHeight,
+			float InVirtualTop, float InVirtualHeight,
+			float InVirtualNestedHeight)
+			: WeakItem(MoveTemp(InWeakItem))
+			, PhysicalTop(InPhysicalTop), PhysicalHeight(InPhysicalHeight)
+			, VirtualTop(InVirtualTop), VirtualHeight(InVirtualHeight), VirtualNestedHeight(InVirtualNestedHeight)
+		{}
+
+		TWeakViewModelPtr<IOutlinerExtension> WeakItem;
+		float PhysicalTop, PhysicalHeight;
+		float VirtualTop, VirtualHeight, VirtualNestedHeight;
+	};
+
+	/** Access all the physical nodes currently visible on the sequencer */
+	const TArray<FCachedGeometry>& GetAllVisibleNodes() const { return PhysicalNodes; }
+
+	/** Ensure that the track area column is either show or hidden, depending on the visibility of the curve editor */
+	void UpdateTrackArea();
+
+	/** Add a SOutlinerView object that should be modified or updated when this Treeview is updated */
+	void AddSlaveTreeView(TSharedPtr<SOutlinerView> SlaveTreeView);
+
+	/** Set a SOutlinerView object this Treeview is slaved to, for operations that should happen on the master */
+	void SetMasterTreeView(TSharedPtr<SOutlinerView> InMasterTreeView) { MasterTreeView = InMasterTreeView; }
+
+	/** Set whether this TreeView should show only pinned nodes or only non-pinned nodes  */
+	void SetShowPinned(bool bShowPinned) { bShowPinnedNodes = bShowPinned; }
+
+protected:
+
+	/** Linear, sorted array of nodes that we currently have generated widgets for */
+	TArray<FCachedGeometry> PhysicalNodes;
+
+protected:
+
+	/** Populate the map of column definitions, and add relevant columns to the header row */
+	void SetupColumns(const FArguments& InArgs);
+
+	FReply OnDragRow(const FGeometry& InGeometry, const FPointerEvent& InPointerEvent, TSharedRef<SOutlinerViewRow> InRow);
+
+	FString OnItemToString_Debug(TWeakViewModelPtr<IOutlinerExtension> InWeakModel);
+
+private:
+
+	/** The tree view's header row (hidden) */
+	TSharedPtr<SHeaderRow> HeaderRow;
+
+	/** The outliner view model */
+	TWeakPtr<FOutlinerViewModel> WeakOutliner;
+
+	/** Cached copy of the root nodes from the tree data */
+	TArray<TWeakViewModelPtr<IOutlinerExtension>> RootNodes;
+
+	/** Column definitions for each of the columns in the tree view */
+	TMap<FName, FOutlinerViewColumn> Columns;
+
+	TAttribute<TSharedPtr<IOutlinerSelectionHandler>> SelectionHandlerAttribute;
+
+	/** Strong pointer to the track area so we can generate track lanes as we need them */
+	TSharedPtr<STrackAreaView> TrackArea;
+
+	/** SOutlinerView objects that should be modified or updated when this Treeview is updated */
+	TArray<TSharedPtr<SOutlinerView>> SlaveTreeViews;
+
+	/** The SOutlinerView object this SOutlinerView is slave to, or nullptr if not a slave */
+	TWeakPtr<SOutlinerView> MasterTreeView;
+
+	float VirtualTop;
+
+	bool bForwardToSelectionHandler;
+
+	/** When true, the tree selection is being updated from a change in the sequencer selection. */
+	bool bUpdatingTreeSelection;
+
+	/** Right mouse button is down, don't update sequencer selection. */
+	bool bRightMouseButtonDown;
+
+	/** Whether this tree is for pinned nodes or non-pinned nodes */
+	bool bShowPinnedNodes;
+
+	/** Whether we have pending selection changes to broadcast */
+	bool bSelectionChangesPending;
+
+	/** Whether physical geometry information should be recomputed */
+	bool bRefreshPhysicalGeometry;
+};
+
+/** Widget that represents a row in the sequencer's tree control. */
+class SEQUENCERCORE_API SOutlinerViewRow
+	: public ISequencerTreeViewRow
+{
+public:
+	DECLARE_DELEGATE_RetVal_ThreeParams(TSharedRef<SWidget>, FOnGenerateWidgetForColumn, TViewModelPtr<IOutlinerExtension>, const FName&, const TSharedRef<SOutlinerViewRow>&);
+	DECLARE_DELEGATE_RetVal_ThreeParams(FReply, FDetectDrag, const FGeometry&, const FPointerEvent&, TSharedRef<SOutlinerViewRow>);
+
+	SLATE_BEGIN_ARGS(SOutlinerViewRow){}
+
+		/** Delegate to invoke to create a new column for this row */
+		SLATE_EVENT(FOnGenerateWidgetForColumn, OnGenerateWidgetForColumn)
+
+		/** Detect a drag on this tree row */
+		SLATE_EVENT(FDetectDrag, OnDetectDrag)
+
+	SLATE_END_ARGS()
+
+	/** Construct function for this widget */
+	void Construct(const FArguments& InArgs, const TSharedRef<STableViewBase>& InOwnerTableView, TWeakViewModelPtr<IOutlinerExtension> InDataModel);
+
+	/** Destroy this widget */
+	~SOutlinerViewRow();
+
+	/** Get the model to which this row relates */
+	TViewModelPtr<IOutlinerExtension> GetDataModel() const;
+
+	/**
+	 * Gets the track lane we relate to
+	 * @param bOnlyOwnTrackLane  Whether to return nullptr if our referenced track lane wasn't created
+	 *							 by our own outliner item view-model 
+	 * @return The track lane for this row
+	 */
+	TSharedPtr<STrackLane> GetTrackLane(bool bOnlyOwnTrackLane = false) const;
+
+	/**
+	 * Adds a reference to track lane, either because:
+	 *  - It is a parent row's track lane that we want to be kept alive when the parent row disappears out of view, or
+	 *  - It is a track lane that was created by our outliner view-model
+	 */
+	void SetTrackLane(const TSharedPtr<STrackLane>& InTrackLane);
+
+	/** Whether the underlying data model is selectable */
+	bool IsSelectable() const;
+
+	virtual const FSlateBrush* GetBorder() const override;
+
+	virtual void Tick(const FGeometry& AllottedGeometry, const double InCurrentTime, const float InDeltaTime) override;
+	
+	virtual int32 OnPaintDropIndicator(EItemDropZone InItemDropZone, const FPaintArgs& Args, const FGeometry& AllottedGeometry, const FSlateRect& MyCullingRect, FSlateWindowElementList& OutDrawElements, int32 LayerId, const FWidgetStyle& InWidgetStyle, bool bParentEnabled) const override;
+
+	/** Overridden from SMultiColumnTableRow.  Generates a widget for this column of the tree row. */
+	virtual TSharedRef<SWidget> GenerateWidgetForColumn(const FName& ColumnId) override;
+
+	/** Called whenever a drag is detected by the tree view. */
+	FReply OnDragDetected(const FGeometry& InGeometry, const FPointerEvent& InPointerEvent) override;
+
+	/** Called to determine whether a current drag operation is valid for this row. */
+	TOptional<EItemDropZone> OnCanAcceptDrop( const FDragDropEvent& DragDropEvent, EItemDropZone ItemDropZone, TWeakViewModelPtr<IOutlinerExtension> InDataModel);
+
+	/** Called to complete a drag and drop onto this drop. */
+	FReply OnAcceptDrop( const FDragDropEvent& DragDropEvent, EItemDropZone ItemDropZone, TWeakViewModelPtr<IOutlinerExtension> InDataModel);
+
+private:
+
+	/**
+	 * Cached reference to a track lane that we relate to.
+	 * Depending on the situation, this is either the actual track lane originally created
+	 * by the our outliner model, or it's just a reference we use to keep the track lane alive
+	 * (it's a weak widget) as long as we are in view.
+	 */
+	TSharedPtr<STrackLane> TrackLane;
+
+	/** The item associated with this row of data */
+	TWeakViewModelPtr<IOutlinerExtension> WeakModel;
+
+	/** Delegate to call to create a new widget for a particular column. */
+	FOnGenerateWidgetForColumn OnGenerateWidgetForColumn;
+
+	/** Delegate to call when dragging a tree item is detected */
+	FDetectDrag OnDetectDrag;
+};
+
+} // namespace Sequencer
+} // namespace UE
+
diff --git a/Engine/Source/Editor/SequencerCore/Public/MVVM/Views/STrackAreaView.h b/Engine/Source/Editor/SequencerCore/Public/MVVM/Views/STrackAreaView.h
new file mode 100644
--- /dev/null
+++ b/Engine/Source/Editor/SequencerCore/Public/MVVM/Views/STrackAreaView.h
@@ -0,0 +1,171 @@
+// Copyright Epic Games, Inc. All Rights Reserved.
+
+#pragma once
+
+#include "CoreTypes.h"
+#include "Layout/Geometry.h"
+#include "Input/CursorReply.h"
+#include "Input/Reply.h"
+#include "Widgets/DeclarativeSyntaxSupport.h"
+#include "SlotBase.h"
+#include "Layout/Children.h"
+#include "Widgets/SPanel.h"
+#include "MVVM/Views/SequencerInputHandlerStack.h"
+#include "MVVM/ViewModels/ViewModel.h"
+#include "MVVM/ViewModelPtr.h"
+
+class FArrangedChildren;
+class FPaintArgs;
+class FSlateWindowElementList;
+class ISequencerEditTool;
+
+namespace UE
+{
+namespace Sequencer
+{
+
+class SOutlinerView;
+class IOutlinerExtension;
+class FTrackAreaViewModel;
+class FViewModel;
+class STrackLane;
+
+/**
+ * Structure representing a slot in the track area.
+ */
+class FTrackAreaSlot
+	: public TSlotBase<FTrackAreaSlot>, public TAlignmentWidgetSlotMixin<FTrackAreaSlot>
+{
+public:
+
+	/** Construction from a track lane */
+	FTrackAreaSlot(const TSharedPtr<STrackLane>& InSlotContent);
+
+	/** The track lane that we represent. */
+	TWeakPtr<STrackLane> TrackLane;
+};
+
+
+/**
+ * The area where tracks( rows of sections ) are displayed.
+ */
+class SEQUENCERCORE_API STrackAreaView
+	: public SPanel
+{
+public:
+
+	SLATE_BEGIN_ARGS( STrackAreaView )
+	{
+		_Clipping = EWidgetClipping::ClipToBounds;
+	}
+	SLATE_END_ARGS()
+
+	STrackAreaView();
+	~STrackAreaView();
+
+	/** Construct this widget */
+	void Construct(const FArguments& InArgs, TWeakPtr<FTrackAreaViewModel> InWeakViewModel);
+
+	TSharedPtr<FTrackAreaViewModel> GetViewModel() const;
+	
+	void SetVirtualPosition(float InVirtualTop);
+
+public:
+
+	/** Empty the track area */
+	void Empty();
+
+	/** Add a new track slot to this area for the given node. The slot will be automatically cleaned up when all external references to the supplied slot are removed. */
+	void AddTrackSlot(const TViewModelPtr<IOutlinerExtension>& InDataModel, const TSharedPtr<STrackLane>& InSlot);
+
+	/** Attempt to find an existing slot relating to the given node */
+	TSharedPtr<STrackLane> FindTrackSlot(const TViewModelPtr<IOutlinerExtension>& InDataModel);
+
+	/** Assign a tree view to this track area. */
+	void SetOutliner(const TSharedPtr<SOutlinerView>& InOutliner);
+	TWeakPtr<SOutlinerView> GetOutliner() const { return WeakOutliner.Pin(); }
+
+	/** Set whether this TrackArea should show only pinned nodes or only non-pinned nodes  */
+	void SetShowPinned(bool bShowPinned) { bShowPinnedNodes = bShowPinned; }
+	bool ShowPinned() const { return bShowPinnedNodes; }
+
+	/** Set whether this TrackArea is slaved to a master TrackArea and should skip updating external controls */
+	void SetIsSlave(bool bInIsSlave) { bIsSlave = bInIsSlave; }
+	bool IsSlave() const { return bIsSlave; }
+
+	static FLinearColor BlendDefaultTrackColor(FLinearColor InColor);
+
+public:
+
+	/*~ SWidget interface */
+	FReply OnMouseButtonDown( const FGeometry& MyGeometry, const FPointerEvent& MouseEvent ) override;
+	FReply OnMouseButtonUp( const FGeometry& MyGeometry, const FPointerEvent& MouseEvent ) override;
+	FReply OnMouseMove( const FGeometry& MyGeometry, const FPointerEvent& MouseEvent ) override;
+	FReply OnMouseWheel( const FGeometry& MyGeometry, const FPointerEvent& MouseEvent ) override;
+	void OnMouseEnter(const FGeometry& MyGeometry, const FPointerEvent& MouseEvent) override;
+	void OnMouseLeave(const FPointerEvent& MouseEvent) override;
+	int32 OnPaint(const FPaintArgs& Args, const FGeometry& AllottedGeometry, const FSlateRect& MyCullingRect, FSlateWindowElementList& OutDrawElements, int32 LayerId, const FWidgetStyle& InWidgetStyle, bool bParentEnabled ) const override;
+	void OnMouseCaptureLost(const FCaptureLostEvent& CaptureLostEvent) override;
+	FCursorReply OnCursorQuery( const FGeometry& MyGeometry, const FPointerEvent& CursorEvent ) const override;
+	void Tick( const FGeometry& AllottedGeometry, const double InCurrentTime, const float InDeltaTime ) override;
+	void OnArrangeChildren( const FGeometry& AllottedGeometry, FArrangedChildren& ArrangedChildren ) const override;
+	FVector2D ComputeDesiredSize(float) const override;
+	FChildren* GetChildren() override;
+	void OnDragEnter(const FGeometry& MyGeometry, const FDragDropEvent& DragDropEvent) override;
+	void OnDragLeave(const FDragDropEvent& DragDropEvent) override;
+	FReply OnDragOver(const FGeometry& MyGeometry, const FDragDropEvent& DragDropEvent) override;
+	FReply OnDrop(const FGeometry& MyGeometry, const FDragDropEvent& DragDropEvent) override;
+
+protected:
+
+	virtual void OnResized(const FVector2D& OldSize, const FVector2D& NewSize){}
+
+	/** Update any hover state required for the track area */
+	virtual void UpdateHoverStates( const FGeometry& MyGeometry, const FPointerEvent& MouseEvent );
+
+protected:
+
+
+private:
+
+	/** The track area's children. */
+	TPanelChildren<FTrackAreaSlot> Children;
+
+protected:
+
+	/** Input handler stack responsible for routing input to the different handlers */
+	FInputHandlerStack InputStack;
+
+	/** A map of child slot content that exist in our view. */
+	TMap<TWeakViewModelPtr<IOutlinerExtension>, TWeakPtr<STrackLane>> TrackSlots;
+
+	/** Weak pointer to the track area view model. */
+	TWeakPtr<FTrackAreaViewModel> WeakViewModel;
+
+	/** Weak pointer to the outliner view (used for scrolling interactions). */
+	TWeakPtr<SOutlinerView> WeakOutliner;
+
+	/** Keep a hold of the size of the area so we can maintain zoom levels. */
+	TOptional<FVector2D> SizeLastFrame;
+
+	/** Weak pointer to the dropped node */
+	TWeakViewModelPtr<IOutlinerExtension> WeakDroppedItem;
+
+	float VirtualTop;
+
+	/** Whether the dropped node is allowed to be dropped onto */
+	bool bAllowDrop;
+
+	/** The frame range of the section about to be dropped */
+	TOptional<TRange<FFrameNumber>> DropFrameRange;
+
+	/** Whether this TrackArea is for pinned nodes or non-pinned nodes */
+	bool bShowPinnedNodes;
+
+	/** Whether this TrackArea is slaved to a master TrackArea and should skip updating external controls */
+	bool bIsSlave;
+};
+
+} // namespace Sequencer
+} // namespace UE
+
diff --git a/Engine/Source/Editor/SequencerCore/Public/MVVM/Views/ViewUtilities.h b/Engine/Source/Editor/SequencerCore/Public/MVVM/Views/ViewUtilities.h
new file mode 100644
--- /dev/null
+++ b/Engine/Source/Editor/SequencerCore/Public/MVVM/Views/ViewUtilities.h
@@ -0,0 +1,19 @@
+// Copyright Epic Games, Inc. All Rights Reserved.
+
+#pragma once
+
+#include "CoreTypes.h"
+#include "Widgets/SWidget.h"
+#include "Framework/SlateDelegates.h"
+
+template< typename ObjectType > class TAttribute;
+
+namespace UE
+{
+namespace Sequencer
+{
+
+SEQUENCERCORE_API TSharedRef<SWidget> MakeAddButton(FText HoverText, FOnGetContent MenuContent, const TAttribute<bool>& HoverState, const TAttribute<bool>& IsEnabled);
+
+} // namespace Sequencer
+} // namespace UE
\ No newline at end of file
diff --git a/Engine/Source/Editor/SequencerCore/Public/TimeToPixel.h b/Engine/Source/Editor/SequencerCore/Public/TimeToPixel.h
new file mode 100644
--- /dev/null
+++ b/Engine/Source/Editor/SequencerCore/Public/TimeToPixel.h
@@ -0,0 +1,137 @@
+// Copyright Epic Games, Inc. All Rights Reserved.
+
+#pragma once
+
+#include "CoreTypes.h"
+#include "Layout/Geometry.h"
+#include "Misc/FrameRate.h"
+
+//TODO: Move to UE::Sequencer namespace
+
+/**
+ * Utility for converting time units to slate pixel units and vice versa
+ */
+struct FTimeToPixel
+{
+public:
+
+	FTimeToPixel( const FGeometry& AllottedGeometry, const TRange<double>& InLocalViewRange, const FFrameRate& InTickResolution )
+		: FTimeToPixel(AllottedGeometry.GetLocalSize().X, InLocalViewRange, InTickResolution)
+	{}
+
+	FTimeToPixel( float WidthPx, const TRange<double>& InLocalViewRange, const FFrameRate& InTickResolution )
+		: ViewRangeStartSeconds( InLocalViewRange.GetLowerBoundValue() )
+		, TickResolution( InTickResolution )
+	{
+		double VisibleWidth = InLocalViewRange.Size<double>();
+
+		const float MaxPixelsPerSecond = 1000.f;
+		PixelsPerSecond = VisibleWidth > 0 ? WidthPx / VisibleWidth : MaxPixelsPerSecond;
+	}
+
+	/**
+	 * Converts a time to a pixel point relative to the geometry of a widget (passed into the constructor)
+	 *
+	 * @param Time	The time to convert
+	 * @return The pixel equivalent of the time
+	 */
+	float SecondsToPixel( double Time ) const
+	{
+		return (Time - ViewRangeStartSeconds) * PixelsPerSecond;
+	}
+
+	/**
+	 * Converts a time delta to a pixel delta
+	 *
+	 * @param TimeDelta  The time delta to convert
+	 * @return           The pixel equivalent of the delta time
+	 */
+	float SecondsDeltaToPixel( double TimeDelta ) const
+	{
+		return TimeDelta * PixelsPerSecond;
+	}
+
+	/**
+	 * Converts a pixel value to time 
+	 *
+	 * @param PixelX The x value of a pixel coordinate relative to the geometry that was passed into the constructor.
+	 * @return The time where the pixel is located
+	 */
+	double PixelToSeconds( float PixelX ) const
+	{
+		return (PixelX/PixelsPerSecond) + ViewRangeStartSeconds;
+	}
+
+	/**
+	 * Converts a frame time to a pixel point relative to the geometry of a widget (passed into the constructor)
+	 *
+	 * @param Time The time to convert
+	 * @return The pixel equivalent of the frame time
+	 */
+	float FrameToPixel( const FFrameTime& Time ) const
+	{
+		return (Time / TickResolution - ViewRangeStartSeconds) * PixelsPerSecond;
+	}
+
+	/**
+	 * Converts a frame delta value to pixel delta
+	 *
+	 * @param TimeDelta   The time delta to convert
+	 * @return            The pixel equivalent of the delta time
+	 */
+	float FrameDeltaToPixel( const FFrameTime& TimeDelta) const
+	{
+		return (TimeDelta / TickResolution) * PixelsPerSecond;
+	}
+
+	/**
+	 * Converts a pixel value to frame time 
+	 *
+	 * @param PixelX The x value of a pixel coordinate relative to the geometry that was passed into the constructor.
+	 * @return The frame time where the pixel is located
+	 */
+	FFrameTime PixelToFrame( float PixelX ) const
+	{
+		return ( PixelX / PixelsPerSecond + ViewRangeStartSeconds ) * TickResolution;
+	}
+
+	/**
+	 * Converts a pixel delta value to delta frame time 
+	 *
+	 * @param PixelDelta The delta value in pixel space
+	 * @return The equivalent delta frame time
+	 */
+	FFrameTime PixelDeltaToFrame( float PixelDelta ) const
+	{
+		return ( PixelDelta / PixelsPerSecond ) * TickResolution;
+	}
+
+	/**
+	 * Retrieve the tick resolution of the current sequence
+	 */
+	FFrameRate GetTickResolution() const
+	{
+		return TickResolution;
+	}
+
+	/**
+	 * Make this converter relative to the specified time (ie, such that pixel 0 == FrameAmount)
+	 */
+	FTimeToPixel RelativeTo(FFrameNumber FrameZero) const
+	{
+		FTimeToPixel Copy = *this;
+		Copy.ViewRangeStartSeconds = FrameZero / TickResolution;
+		return Copy;
+	}
+
+private:
+	FTimeToPixel(){}
+
+	/** time range of the sequencer in seconds */
+	double ViewRangeStartSeconds;
+	/** The tick resolution of the current timeline */
+	FFrameRate TickResolution;
+	/** The number of pixels in the view range */
+	float PixelsPerSecond;
+
+};
diff --git a/Engine/Source/Editor/SequencerWidgets/Public/ITimeSlider.h b/Engine/Source/Editor/SequencerWidgets/Public/ITimeSlider.h
--- a/Engine/Source/Editor/SequencerWidgets/Public/ITimeSlider.h
+++ b/Engine/Source/Editor/SequencerWidgets/Public/ITimeSlider.h
@@ -9,7 +9,7 @@
 #include "Widgets/DeclarativeSyntaxSupport.h"
 #include "Widgets/SCompoundWidget.h"
 #include "Widgets/Input/NumericTypeInterface.h"
-#include "Editor/Sequencer/Public/ISequencerInputHandler.h"
+#include "Editor/SequencerCore/Public/ISequencerInputHandler.h"
 #include "IMovieScenePlayer.h"
 #include "MovieScene.h"
 
diff --git a/Engine/Source/Programs/IncludeTool/IncludeTool/Rules.cs b/Engine/Source/Programs/IncludeTool/IncludeTool/Rules.cs
--- a/Engine/Source/Programs/IncludeTool/IncludeTool/Rules.cs
+++ b/Engine/Source/Programs/IncludeTool/IncludeTool/Rules.cs
@@ -419,7 +419,8 @@
 			"/Engine/Source/Runtime/SlateCore/Public/Fonts/ShapedTextFwd.h", // Typedef isn't a forward declaration
 			"/Engine/Source/Runtime/Slate/Public/Framework/Text/ShapedTextCacheFwd.h", // Typedef isn't a forward declaration
 			"/Engine/Source/Runtime/MovieScene/Public/MovieSceneFwd.h",
-			"/Engine/Source/Runtime/Core/Public/Math/MathFwd.h", // invalid forward declaration - 'namespace UE::Math'
+            "/Engine/Source/Editor/SequencerCore/Public/SequencerCoreFwd.h", // invalid forward declaration - 'namespace UE'
+            "/Engine/Source/Runtime/Core/Public/Math/MathFwd.h", // invalid forward declaration - 'namespace UE::Math'
 			"/Engine/Source/Runtime/Core/Public/Containers/ContainersFwd.h", // invalid forward declaration - 'template<> struct TIsContiguousContainer<Type> { static constexpr bool Value = true; };'
 			"/Engine/Source/Runtime/Core/Public/Containers/StringFwd.h", // invalid forward declaration - 'template<> struct TIsContiguousContainer<Type> { static constexpr bool Value = true; };'
 			"/Engine/Source/Runtime/Core/Public/Internationalization/StringTableCoreFwd.h", // Typedef isn't a forward declaration
diff --git a/Engine/Source/Runtime/LevelSequence/Private/LevelSequence.cpp b/Engine/Source/Runtime/LevelSequence/Private/LevelSequence.cpp
--- a/Engine/Source/Runtime/LevelSequence/Private/LevelSequence.cpp
+++ b/Engine/Source/Runtime/LevelSequence/Private/LevelSequence.cpp
@@ -605,7 +605,7 @@
 		FMovieScenePossessable* ChildPossessable = MovieScene->FindPossessable(NewGuid);
 		if (ensure(ChildPossessable))
 		{
-			ChildPossessable->SetParent(ParentGuid);
+			ChildPossessable->SetParent(ParentGuid, MovieScene);
 		}
 
 		FMovieSceneSpawnable* ParentSpawnable = MovieScene->FindSpawnable(ParentGuid);
diff --git a/Engine/Source/Runtime/LevelSequence/Public/LevelSequenceActor.h b/Engine/Source/Runtime/LevelSequence/Public/LevelSequenceActor.h
--- a/Engine/Source/Runtime/LevelSequence/Public/LevelSequenceActor.h
+++ b/Engine/Source/Runtime/LevelSequence/Public/LevelSequenceActor.h
@@ -76,13 +76,13 @@
 
 public:
 
-	UPROPERTY(EditAnywhere, BlueprintReadOnly, Category="Playback", meta=(ShowOnlyInnerProperties))
+	UPROPERTY(EditAnywhere, BlueprintReadOnly, Category="Playback", meta=(ShowOnlyInnerProperties, ExposeOnSpawn))
 	FMovieSceneSequencePlaybackSettings PlaybackSettings;
 
 	UPROPERTY(Instanced, transient, replicated, BlueprintReadOnly, BlueprintGetter=GetSequencePlayer, Category="Playback", meta=(ExposeFunctionCategories="Sequencer|Player"))
 	TObjectPtr<ULevelSequencePlayer> SequencePlayer;
 
-	UPROPERTY(EditAnywhere, replicated, BlueprintReadOnly, Category="General", meta=(AllowedClasses="/Script/LevelSequence.LevelSequence"))
+	UPROPERTY(EditAnywhere, replicated, BlueprintReadOnly, Category="General", meta=(AllowedClasses="/Script/LevelSequence.LevelSequence", ExposeOnSpawn))
 	TObjectPtr<ULevelSequence> LevelSequenceAsset;
 
 #if WITH_EDITORONLY_DATA
@@ -90,7 +90,7 @@
 	FSoftObjectPath LevelSequence_DEPRECATED;
 #endif
 
-	UPROPERTY(EditAnywhere, BlueprintReadWrite, Category="Cameras", meta=(ShowOnlyInnerProperties))
+	UPROPERTY(EditAnywhere, BlueprintReadWrite, Category="Cameras", meta=(ShowOnlyInnerProperties, ExposeOnSpawn))
 	FLevelSequenceCameraSettings CameraSettings;
 
 	UPROPERTY(Instanced, BlueprintReadOnly, Category="General")
@@ -108,7 +108,7 @@
 	uint8 bOverrideInstanceData : 1;
 
 	/** If true, playback of this level sequence on the server will be synchronized across other clients */
-	UPROPERTY(EditAnywhere, DisplayName="Replicate Playback", BlueprintReadWrite, BlueprintSetter=SetReplicatePlayback, Category=Replication)
+	UPROPERTY(EditAnywhere, DisplayName="Replicate Playback", BlueprintReadWrite, BlueprintSetter=SetReplicatePlayback, Category=Replication, meta=(ExposeOnSpawn))
 	uint8 bReplicatePlayback:1;
 
 	/** Instance data that can be used to dynamically control sequence evaluation at runtime */
diff --git a/Engine/Source/Runtime/MovieScene/Private/MovieScene.cpp b/Engine/Source/Runtime/MovieScene/Private/MovieScene.cpp
--- a/Engine/Source/Runtime/MovieScene/Private/MovieScene.cpp
+++ b/Engine/Source/Runtime/MovieScene/Private/MovieScene.cpp
@@ -11,7 +11,7 @@
 #include "UObject/SequencerObjectVersion.h"
 #include "Evaluation/IMovieSceneCustomClockSource.h"
 #include "CommonFrameRates.h"
-#include "EntitySystem/IMovieSceneEntityProvider.h"
+#include "EventHandlers/ISequenceDataEventHandler.h"
 #include "Misc/FrameRate.h"
 #include "UObject/ObjectSaveContext.h"
 #include "UObject/UObjectHash.h"
@@ -206,6 +206,8 @@
 	// Add a new binding so that tracks can be added to it
 	new (ObjectBindings) FMovieSceneBinding( NewSpawnable.GetGuid(), NewSpawnable.GetName() );
 
+	EventHandlers.Trigger(&UE::MovieScene::ISequenceDataEventHandler::OnBindingAdded, ObjectBindings.Last());
+
 	return NewSpawnable.GetGuid();
 }
 
@@ -224,6 +226,8 @@
 		Track->Rename(nullptr, this);
 	}
 	ObjectBindings.Add(NewBinding);
+
+	EventHandlers.Trigger(&UE::MovieScene::ISequenceDataEventHandler::OnBindingAdded, ObjectBindings.Last());
 }
 
 bool UMovieScene::RemoveSpawnable( const FGuid& Guid )
@@ -241,6 +245,8 @@
 
 				Spawnables.RemoveAt( SpawnableIter.GetIndex() );
 
+				EventHandlers.Trigger(&UE::MovieScene::ISequenceDataEventHandler::OnBindingRemoved, Guid);
+
 				bAnythingRemoved = true;
 				break;
 			}
@@ -283,6 +289,8 @@
 	// Add a new binding so that tracks can be added to it
 	new (ObjectBindings) FMovieSceneBinding( NewPossessable.GetGuid(), NewPossessable.GetName() );
 
+	EventHandlers.Trigger(&UE::MovieScene::ISequenceDataEventHandler::OnBindingAdded, ObjectBindings.Last());
+
 	return NewPossessable.GetGuid();
 }
 
@@ -301,6 +309,8 @@
 		Track->Rename(nullptr, this);
 	}
 	ObjectBindings.Add(NewBinding);
+
+	EventHandlers.Trigger(&UE::MovieScene::ISequenceDataEventHandler::OnBindingAdded, ObjectBindings.Last());
 }
 
 
@@ -331,6 +341,8 @@
 
 			RemoveBinding( PossessableGuid );
 
+			EventHandlers.Trigger(&UE::MovieScene::ISequenceDataEventHandler::OnBindingRemoved, PossessableGuid);
+
 			bAnythingRemoved = true;
 			break;
 		}
@@ -481,10 +493,71 @@
 }
 
 
-TArray<UMovieSceneFolder*>&  UMovieScene::GetRootFolders()
+TArrayView<UMovieSceneFolder* const> UMovieScene::GetRootFolders()
 {
 	return RootFolders;
 }
+
+void UMovieScene::GetRootFolders(TArray<UMovieSceneFolder*>& InRootFolders)
+{
+	InRootFolders.Append(RootFolders);
+}
+
+int32 UMovieScene::GetNumRootFolders() const
+{
+	return RootFolders.Num();
+}
+
+UMovieSceneFolder* UMovieScene::GetRootFolder(int32 FolderIndex) const
+{
+	return RootFolders[FolderIndex];
+}
+
+void UMovieScene::AddRootFolder(UMovieSceneFolder* Folder)
+{
+	if (!RootFolders.Contains(Folder))
+	{
+		RootFolders.Add(Folder);
+		EventHandlers.Trigger(&UE::MovieScene::ISequenceDataEventHandler::OnRootFolderAdded, Folder);
+	}
+}
+
+int32 UMovieScene::RemoveRootFolder(UMovieSceneFolder* Folder)
+{
+	const int32 NumRemoved = RootFolders.Remove(Folder);
+	if (NumRemoved != 0)
+	{
+		EventHandlers.Trigger(&UE::MovieScene::ISequenceDataEventHandler::OnRootFolderRemoved, Folder);
+	}
+	return NumRemoved;
+}
+
+bool UMovieScene::RemoveRootFolder(int32 FolderIndex)
+{
+	if (RootFolders.IsValidIndex(FolderIndex))
+	{
+		UMovieSceneFolder* Folder = RootFolders[FolderIndex];
+		RootFolders.RemoveAt(FolderIndex);
+		EventHandlers.Trigger(&UE::MovieScene::ISequenceDataEventHandler::OnRootFolderRemoved, Folder);
+		return true;
+	}
+	return false;
+}
+
+void UMovieScene::EmptyRootFolders()
+{
+	TArray<TObjectPtr<UMovieSceneFolder>> OldFolders;
+	Swap(RootFolders, OldFolders);
+
+	RootFolders.Empty();
+
+	for (UMovieSceneFolder* Folder : OldFolders)
+	{
+		EventHandlers.Trigger(&UE::MovieScene::ISequenceDataEventHandler::OnRootFolderRemoved, Folder);
+	}
+}
+
+
 #endif
 
 void UMovieScene::SetPlaybackRange(FFrameNumber Start, int32 Duration, bool bAlwaysMarkDirty)
@@ -955,7 +1028,7 @@
 
 			CreatedType = NewObject<UMovieSceneTrack>(this, TrackClass, NAME_None, RF_Transactional);
 			check(CreatedType);
-			Binding.AddTrack( *CreatedType );
+			Binding.AddTrack( *CreatedType, this );
 		}
 	}
 
@@ -980,7 +1053,7 @@
 		if (Binding.GetObjectGuid() == ObjectGuid)
 		{
 			InTrack->Rename(nullptr, this);
-			Binding.AddTrack(*InTrack);
+			Binding.AddTrack(*InTrack, this);
 			return true;
 		}
 	}
@@ -996,7 +1069,7 @@
 
 	for (auto& Binding : ObjectBindings)
 	{
-		if (Binding.RemoveTrack(Track))
+		if (Binding.RemoveTrack(Track, this))
 		{
 			bAnythingRemoved = true;
 
@@ -1056,6 +1129,9 @@
 
 	UMovieSceneTrack* CreatedType = NewObject<UMovieSceneTrack>(this, TrackClass, NAME_None, RF_Transactional);
 	MasterTracks.Add( CreatedType );
+
+	EventHandlers.Trigger(&UE::MovieScene::ISequenceDataEventHandler::OnMasterTrackAdded, CreatedType);
+
 	return CreatedType;
 }
 
@@ -1074,6 +1150,9 @@
 		Modify();
 		MasterTracks.Add(InTrack);
 		InTrack->Rename(nullptr, this);
+
+		EventHandlers.Trigger(&UE::MovieScene::ISequenceDataEventHandler::OnMasterTrackAdded, InTrack);
+
 		return true;
 	}
 	return false;
@@ -1084,7 +1163,12 @@
 {
 	Modify();
 
-	return (MasterTracks.RemoveSingle(&Track) != 0);
+	const bool bRemoved = MasterTracks.RemoveSingle(&Track) != 0;
+	if (bRemoved)
+	{
+		EventHandlers.Trigger(&UE::MovieScene::ISequenceDataEventHandler::OnMasterTrackRemoved, &Track);
+	}
+	return bRemoved;
 }
 
 
@@ -1115,6 +1199,8 @@
 	{
 		Modify();
 		CameraCutTrack = NewObject<UMovieSceneTrack>(this, TrackClass, NAME_None, RF_Transactional);
+
+		EventHandlers.Trigger(&UE::MovieScene::ISequenceDataEventHandler::OnMasterTrackAdded, CameraCutTrack);
 	}
 
 	return CameraCutTrack;
@@ -1132,15 +1218,31 @@
 	if( CameraCutTrack )
 	{
 		Modify();
+		UMovieSceneTrack* TmpCameraCut = CameraCutTrack;
 		CameraCutTrack = nullptr;
+
+		EventHandlers.Trigger(&UE::MovieScene::ISequenceDataEventHandler::OnMasterTrackRemoved, TmpCameraCut);
 	}
 }
 
 void UMovieScene::SetCameraCutTrack(UMovieSceneTrack* InTrack)
 {
+	if (!InTrack)
+	{
+		return;
+	}
+
 	Modify();
 	InTrack->Rename(nullptr, this);
+	UMovieSceneTrack* OldCameraCutTrack = CameraCutTrack;
 	CameraCutTrack = InTrack;
+
+	if (OldCameraCutTrack)
+	{
+		EventHandlers.Trigger(&UE::MovieScene::ISequenceDataEventHandler::OnMasterTrackRemoved, OldCameraCutTrack);
+	}
+
+	EventHandlers.Trigger(&UE::MovieScene::ISequenceDataEventHandler::OnMasterTrackAdded, CameraCutTrack);
 }
 
 
@@ -1266,19 +1368,9 @@
 		}
 	}
 
-	for ( int32 ObjectBindingIndex = 0; ObjectBindingIndex < ObjectBindings.Num(); ++ObjectBindingIndex)
+	for (FMovieSceneBinding& Binding : ObjectBindings)
 	{
-		for( int32 TrackIndex = 0; TrackIndex < ObjectBindings[ObjectBindingIndex].GetTracks().Num(); )
-		{
-			if (ObjectBindings[ObjectBindingIndex].GetTracks()[TrackIndex] == nullptr)
-			{
-				ObjectBindings[ObjectBindingIndex].RemoveTrack(*ObjectBindings[ObjectBindingIndex].GetTracks()[TrackIndex]);
-			}
-			else
-			{
-				++TrackIndex;
-			}
-		}
+		Binding.RemoveNullTracks();
 	}
 
 #if WITH_EDITORONLY_DATA
@@ -1358,7 +1450,9 @@
 
 void UMovieScene::RemoveBinding(const FGuid& Guid)
 {
-	// update each type
+	// WARNING: This function intentionally does not trigger events to ensure
+	// that events are triggered when all processing is complete (ie, when removing a spawnable or posessable)
+
 	for (int32 BindingIndex = 0; BindingIndex < ObjectBindings.Num(); ++BindingIndex)
 	{
 		if (ObjectBindings[BindingIndex].GetObjectGuid() == Guid)
@@ -1384,6 +1478,9 @@
 			{
 				Track->Modify();
 			}
+
+			EventHandlers.Trigger(&UE::MovieScene::ISequenceDataEventHandler::OnBindingRemoved, OldGuid);
+			EventHandlers.Trigger(&UE::MovieScene::ISequenceDataEventHandler::OnBindingAdded, Binding);
 			break;
 		}
 	}
@@ -1403,6 +1500,9 @@
 		{
 			Track->Rename(nullptr, this);
 		}
+
+		EventHandlers.Trigger(&UE::MovieScene::ISequenceDataEventHandler::OnBindingRemoved, BindingToReplaceGuid);
+		EventHandlers.Trigger(&UE::MovieScene::ISequenceDataEventHandler::OnBindingAdded, *Binding);
 	}
 }
 
@@ -1431,7 +1531,7 @@
 	if (SourceBinding && DestinationBinding)
 	{
 		// Swap the tracks round
-		DestinationBinding->SetTracks(SourceBinding->StealTracks());
+		DestinationBinding->SetTracks(SourceBinding->StealTracks(this), this);
 
 		// Changing a binding guid invalidates any tracks contained within the binding
 		// Make sure they are written into the transaction buffer by calling modify
@@ -1472,7 +1572,7 @@
 
 			if (!bAlreadyExists)
 			{
-				SourcePossessable.SetParent(DestinationBindingId);
+				SourcePossessable.SetParent(DestinationBindingId, this);
 				if (DestinationSpawnable)
 				{
 					DestinationSpawnable->AddChildPossessable(SourcePossessable.GetGuid());
diff --git a/Engine/Source/Runtime/MovieScene/Private/MovieSceneBinding.cpp b/Engine/Source/Runtime/MovieScene/Private/MovieSceneBinding.cpp
--- a/Engine/Source/Runtime/MovieScene/Private/MovieSceneBinding.cpp
+++ b/Engine/Source/Runtime/MovieScene/Private/MovieSceneBinding.cpp
@@ -1,25 +1,100 @@
-// Copyright Epic Games, Inc. All Rights Reserved.
-
-#include "MovieSceneBinding.h"
-#include "MovieSceneTrack.h"
-#include "MovieScene.h"
-
-/* FMovieSceneBinding interface
- *****************************************************************************/
-
-void FMovieSceneBinding::AddTrack(UMovieSceneTrack& NewTrack)
-{
-#if WITH_EDITOR
-	if (!UMovieScene::IsTrackClassAllowed(NewTrack.GetClass()))
-	{
-		return;
-	}
-#endif
-
-	Tracks.Add(&NewTrack);
-}
-
-bool FMovieSceneBinding::RemoveTrack(UMovieSceneTrack& Track)
-{
-	return (Tracks.RemoveSingle(&Track) != 0);
-}
+// Copyright Epic Games, Inc. All Rights Reserved.
+
+#include "MovieSceneBinding.h"
+#include "MovieSceneTrack.h"
+
+#include "MovieScene.h"
+#include "EventHandlers/ISequenceDataEventHandler.h"
+
+/* FMovieSceneBinding interface
+ *****************************************************************************/
+
+void FMovieSceneBinding::AddTrack(UMovieSceneTrack& NewTrack, UMovieScene* Owner)
+{
+#if WITH_EDITOR
+	if (!UMovieScene::IsTrackClassAllowed(NewTrack.GetClass()))
+	{
+		return;
+	}
+#endif
+
+	Tracks.Add(&NewTrack);
+
+	Owner->EventHandlers.Trigger(&UE::MovieScene::ISequenceDataEventHandler::OnTrackAddedToBinding, &NewTrack, ObjectGuid);
+}
+
+bool FMovieSceneBinding::RemoveTrack(UMovieSceneTrack& Track, UMovieScene* Owner)
+{
+	if (Tracks.RemoveSingle(&Track) != 0)
+	{
+		Owner->EventHandlers.Trigger(&UE::MovieScene::ISequenceDataEventHandler::OnTrackRemovedFromBinding, &Track, ObjectGuid);
+		return true;
+	}
+	return false;
+}
+
+void FMovieSceneBinding::RemoveNullTracks()
+{
+	for (int32 TrackIndex = Tracks.Num()-1; TrackIndex >= 0; --TrackIndex)
+	{
+		if (Tracks[TrackIndex] == nullptr)
+		{
+			Tracks.RemoveAt(TrackIndex);
+			// Don't trigger events for null tracks
+		}
+	}
+}
+
+TArray<UMovieSceneTrack*> FMovieSceneBinding::StealTracks(UMovieScene* Owner)
+{
+	if (Owner)
+	{
+		for (UMovieSceneTrack* Track : Tracks)
+		{
+			Owner->EventHandlers.Trigger(&UE::MovieScene::ISequenceDataEventHandler::OnTrackRemovedFromBinding, Track, ObjectGuid);
+		}
+	}
+
+	decltype(Tracks) Empty;
+	Swap(Empty, Tracks);
+	return Empty;
+}
+
+void FMovieSceneBinding::SetTracks(TArray<UMovieSceneTrack*>&& InTracks, UMovieScene* Owner)
+{
+	// Care is taken here to ensure that we trigger the events correctly when
+	// InTracks contains tracks that already exist in this binding
+	TSet<UMovieSceneTrack*> NewTracks;
+	for (UMovieSceneTrack* Track : InTracks)
+	{
+		NewTracks.Add(Track);
+	}
+
+	for (int32 Index = Tracks.Num()-1; Index >= 0; --Index)
+	{
+		UMovieSceneTrack* Track = Tracks[Index];
+		if (!NewTracks.Contains(Track))
+		{
+			Tracks.RemoveAt(Index, 1, false);
+			if (Owner)
+			{
+				Owner->EventHandlers.Trigger(&UE::MovieScene::ISequenceDataEventHandler::OnTrackRemovedFromBinding, Track, ObjectGuid);
+			}
+		}
+		else
+		{
+			NewTracks.Remove(Track);
+		}
+	}
+
+	for (UMovieSceneTrack* Track : NewTracks)
+	{
+		Tracks.Add(Track);
+		if (Owner)
+		{
+			Owner->EventHandlers.Trigger(&UE::MovieScene::ISequenceDataEventHandler::OnTrackAddedToBinding, Track, ObjectGuid);
+		}
+	}
+
+	Tracks.Shrink();
+}
\ No newline at end of file
diff --git a/Engine/Source/Runtime/MovieScene/Private/MovieSceneFolder.cpp b/Engine/Source/Runtime/MovieScene/Private/MovieSceneFolder.cpp
--- a/Engine/Source/Runtime/MovieScene/Private/MovieSceneFolder.cpp
+++ b/Engine/Source/Runtime/MovieScene/Private/MovieSceneFolder.cpp
@@ -5,7 +5,7 @@
 #include "MovieSceneTrack.h"
 #include "Algo/Count.h"
 
-void GetMovieSceneFoldersRecursive(const TArray<UMovieSceneFolder*>& InFoldersToRecurse, TArray<UMovieSceneFolder*>& OutFolders)
+void GetMovieSceneFoldersRecursive(TArrayView<UMovieSceneFolder* const> InFoldersToRecurse, TArray<UMovieSceneFolder*>& OutFolders)
 {
 	for (UMovieSceneFolder* Folder : InFoldersToRecurse)
 	{
@@ -40,7 +40,7 @@
 }
 
 
-const TArray<UMovieSceneFolder*>& UMovieSceneFolder::GetChildFolders() const
+TArrayView<UMovieSceneFolder* const> UMovieSceneFolder::GetChildFolders() const
 {
 	return ChildFolders;
 }
@@ -67,6 +67,8 @@
 
 	// Now add it as a child of ourself
 	ChildFolders.Add( InChildFolder );
+
+	EventHandlers.Trigger(&UE::MovieScene::IFolderEventHandler::OnChildFolderAdded, InChildFolder);
 }
 
 
@@ -75,6 +77,8 @@
 	Modify();
 
 	ChildFolders.Remove(InChildFolder);
+
+	EventHandlers.Trigger(&UE::MovieScene::IFolderEventHandler::OnChildFolderRemoved, InChildFolder);
 }
 
 
@@ -104,6 +108,8 @@
 #endif
 
 	ChildMasterTracks.Add( InMasterTrack );
+
+	EventHandlers.Trigger(&UE::MovieScene::IFolderEventHandler::OnTrackAdded, InMasterTrack);
 }
 
 
@@ -111,7 +117,10 @@
 {
 	Modify();
 
-	ChildMasterTracks.Remove( InMasterTrack );
+	if (ChildMasterTracks.Remove( InMasterTrack ) > 0)
+	{
+		EventHandlers.Trigger(&UE::MovieScene::IFolderEventHandler::OnTrackRemoved, InMasterTrack);
+	}
 }
 
 
@@ -149,6 +158,8 @@
 #endif
 
 	ChildObjectBindings.Add( InObjectBinding );
+
+	EventHandlers.Trigger(&UE::MovieScene::IFolderEventHandler::OnObjectBindingAdded, InObjectBinding);
 }
 
 
@@ -156,7 +167,10 @@
 {
 	Modify();
 
-	ChildObjectBindings.Remove( InObjectBinding );
+	if (ChildObjectBindings.Remove( InObjectBinding ) > 0)
+	{
+		EventHandlers.Trigger(&UE::MovieScene::IFolderEventHandler::OnObjectBindingRemoved, InObjectBinding);
+	}
 }
 
 void UMovieSceneFolder::ClearChildObjectBindings()
@@ -274,7 +288,7 @@
 	}
 }
 
-void UMovieSceneFolder::CalculateFolderPath(UMovieSceneFolder* Folder, const TArray<UMovieSceneFolder*>& RootFolders, TArray<FName>& FolderPath)
+void UMovieSceneFolder::CalculateFolderPath(UMovieSceneFolder* Folder, TArrayView<UMovieSceneFolder* const> RootFolders, TArray<FName>& FolderPath)
 {
 	TMap<UMovieSceneFolder*, UMovieSceneFolder*> ChildToParentMap;
 	for (UMovieSceneFolder* RootFolder : RootFolders)
@@ -293,7 +307,7 @@
 	}
 }
 
-UMovieSceneFolder* UMovieSceneFolder::GetFolderWithPath(const TArray<FName>& InFolderPath, const TArray<UMovieSceneFolder*>& InFolders, const TArray<UMovieSceneFolder*>& RootFolders)
+UMovieSceneFolder* UMovieSceneFolder::GetFolderWithPath(const TArray<FName>& InFolderPath, const TArray<UMovieSceneFolder*>& InFolders, TArrayView<UMovieSceneFolder* const> RootFolders)
 {
 	for (UMovieSceneFolder* Folder : InFolders)
 	{
@@ -331,3 +345,55 @@
 		Super::Serialize( Archive );
 	}
 }
+
+FName UMovieSceneFolder::MakeUniqueChildFolderName(FName InName) const
+{
+	return MakeUniqueChildFolderName(InName, GetChildFolders());
+}
+
+FName UMovieSceneFolder::MakeUniqueChildFolderName(FName InName, TArrayView<UMovieSceneFolder* const> InFolders)
+{
+	bool bFoundExactDuplicate = false;
+
+	int32 NextNameIndex = InName.GetNumber();
+
+	// Iterate all children, finding a new unique name index for any that have the same base name
+	for (UMovieSceneFolder* Child : InFolders)
+	{
+		constexpr bool bCompareNumber = false;
+		if (InName.IsEqual(Child->GetFolderName(), ENameCase::IgnoreCase, bCompareNumber))
+		{
+			NextNameIndex = FMath::Max(NextNameIndex, Child->GetFolderName().GetNumber()) + 1;
+		}
+
+		if (InName == Child->GetFolderName())
+		{
+			bFoundExactDuplicate = true;
+		}
+	}
+
+	if (bFoundExactDuplicate)
+	{
+		InName.SetNumber(NextNameIndex);
+	}
+
+	return InName;
+}
+
+#if WITH_EDITOR
+
+void UMovieSceneFolder::PostEditUndo()
+{
+	Super::PostEditUndo();
+
+	EventHandlers.Trigger(&UE::MovieScene::IFolderEventHandler::OnPostUndo);
+}
+
+void UMovieSceneFolder::PostEditUndo(TSharedPtr<ITransactionObjectAnnotation> TransactionAnnotation)
+{
+	Super::PostEditUndo(TransactionAnnotation);
+
+	EventHandlers.Trigger(&UE::MovieScene::IFolderEventHandler::OnPostUndo);
+}
+
+#endif
diff --git a/Engine/Source/Runtime/MovieScene/Private/MovieScenePossessable.cpp b/Engine/Source/Runtime/MovieScene/Private/MovieScenePossessable.cpp
--- a/Engine/Source/Runtime/MovieScene/Private/MovieScenePossessable.cpp
+++ b/Engine/Source/Runtime/MovieScene/Private/MovieScenePossessable.cpp
@@ -1,21 +1,36 @@
-// Copyright Epic Games, Inc. All Rights Reserved.
-
-#include "MovieScenePossessable.h"
-#include "MovieSceneSpawnableAnnotation.h"
-#include "MovieSceneSequence.h"
-#include "MovieSceneSequenceID.h"
-#include "MovieScene.h"
-#include "IMovieScenePlayer.h"
-
-bool FMovieScenePossessable::BindSpawnableObject(FMovieSceneSequenceID SequenceID, UObject* Object, IMovieScenePlayer* Player)
-{
-	TOptional<FMovieSceneSpawnableAnnotation> Spawnable = FMovieSceneSpawnableAnnotation::Find(Object);
-	if (Spawnable.IsSet())
-	{
-		// Check whether the spawnable is underneath the current sequence, if so, we can remap it to a local sequence ID
-		SetSpawnableObjectBindingID(UE::MovieScene::FRelativeObjectBindingID(SequenceID, Spawnable->SequenceID, Spawnable->ObjectBindingID, *Player));
-		return true;
-	}
-
-	return false;
-}
+// Copyright Epic Games, Inc. All Rights Reserved.
+
+#include "MovieScenePossessable.h"
+#include "MovieSceneSpawnableAnnotation.h"
+#include "MovieSceneSequence.h"
+#include "MovieSceneSequenceID.h"
+#include "MovieScene.h"
+#include "IMovieScenePlayer.h"
+#include "EventHandlers/ISequenceDataEventHandler.h"
+
+bool FMovieScenePossessable::BindSpawnableObject(FMovieSceneSequenceID SequenceID, UObject* Object, IMovieScenePlayer* Player)
+{
+	TOptional<FMovieSceneSpawnableAnnotation> Spawnable = FMovieSceneSpawnableAnnotation::Find(Object);
+	if (Spawnable.IsSet())
+	{
+		// Check whether the spawnable is underneath the current sequence, if so, we can remap it to a local sequence ID
+		SetSpawnableObjectBindingID(UE::MovieScene::FRelativeObjectBindingID(SequenceID, Spawnable->SequenceID, Spawnable->ObjectBindingID, *Player));
+		return true;
+	}
+
+	return false;
+}
+
+void FMovieScenePossessable::SetParent(const FGuid& InParentGuid)
+{
+	ParentGuid = InParentGuid;
+}
+
+void FMovieScenePossessable::SetParent(const FGuid& InParentGuid, UMovieScene* Owner)
+{
+	if (ParentGuid != InParentGuid)
+	{
+		ParentGuid = InParentGuid;
+		Owner->EventHandlers.Trigger(&UE::MovieScene::ISequenceDataEventHandler::OnBindingParentChanged, Guid, InParentGuid);
+	}
+}
\ No newline at end of file
diff --git a/Engine/Source/Runtime/MovieScene/Private/MovieSceneSection.cpp b/Engine/Source/Runtime/MovieScene/Private/MovieSceneSection.cpp
--- a/Engine/Source/Runtime/MovieScene/Private/MovieSceneSection.cpp
+++ b/Engine/Source/Runtime/MovieScene/Private/MovieSceneSection.cpp
@@ -368,6 +368,17 @@
 	return false;
 }
 
+void UMovieSceneSection::SetRowIndex(int32 NewRowIndex)
+{
+	const int32 OldRowIndex = RowIndex;
+	RowIndex = NewRowIndex;
+
+	if (OldRowIndex != RowIndex)
+	{
+		EventHandlers.Trigger(&UE::MovieScene::ISectionEventHandler::OnRowChanged, this);
+	}
+}
+
 void UMovieSceneSection::GetOverlappingSections(TArray<UMovieSceneSection*>& OutSections, bool bSameRow, bool bIncludeThis)
 {
 	UMovieSceneTrack* Track = GetTypedOuter<UMovieSceneTrack>();
diff --git a/Engine/Source/Runtime/MovieScene/Private/MovieSceneSignedObject.cpp b/Engine/Source/Runtime/MovieScene/Private/MovieSceneSignedObject.cpp
--- a/Engine/Source/Runtime/MovieScene/Private/MovieSceneSignedObject.cpp
+++ b/Engine/Source/Runtime/MovieScene/Private/MovieSceneSignedObject.cpp
@@ -4,12 +4,90 @@
 #include "Templates/Casts.h"
 #include "MovieSceneSequence.h"
 #include "UObject/Package.h"
+#include "CoreGlobals.h"
 
+namespace UE
+{
+namespace MovieScene
+{
+
+static TSet<TWeakPtr<IDeferredSignedObjectFlushSignal>> GDeferredSignedObjectFlushSignals;
+static TUniquePtr<IDeferredSignedObjectChangeHandler> GDeferredSignedObjectChangeHandler;
+static uint32 GScopedSignedObjectDeferCount = 0;
+
+void SignalScopedSignedObjectModifyFlush()
+{
+	// Copy to a temporary set to guard against re-entrancy
+	TSet<TWeakPtr<IDeferredSignedObjectFlushSignal>> Temp;
+	Swap(Temp, GDeferredSignedObjectFlushSignals);
+
+	for (TWeakPtr<UE::MovieScene::IDeferredSignedObjectFlushSignal> WeakSignal : Temp)
+	{
+		if (TSharedPtr<UE::MovieScene::IDeferredSignedObjectFlushSignal> Signal = WeakSignal.Pin())
+		{
+			Signal->OnDeferredModifyFlush();
+		}
+	}
+}
+
+FScopedSignedObjectModifyDefer::FScopedSignedObjectModifyDefer(bool bInForceFlush)
+{
+	++GScopedSignedObjectDeferCount;
+	bForceFlush = bInForceFlush;
+}
+
+FScopedSignedObjectModifyDefer::~FScopedSignedObjectModifyDefer()
+{
+	--GScopedSignedObjectDeferCount;
+
+	if (GScopedSignedObjectDeferCount == 0 || bForceFlush)
+	{
+		if (GDeferredSignedObjectChangeHandler)
+		{
+			GDeferredSignedObjectChangeHandler->Flush();
+		}
+		SignalScopedSignedObjectModifyFlush();
+	}
+}
+
+} // namespace MovieScene
+} // namespace UE
+
+
 UMovieSceneSignedObject::UMovieSceneSignedObject(const FObjectInitializer& Init)
 	: Super(Init)
 {
+	using namespace UE::MovieScene;
+
+ 	if (HasAnyFlags(RF_Transactional) && GDeferredSignedObjectChangeHandler 
+#if WITH_EDITOR
+		&& GIsTransacting
+#endif
+	)
+	{
+		GDeferredSignedObjectChangeHandler->DeferMarkAsChanged(this);
+	}
 }
 
+void UMovieSceneSignedObject::SetDeferredHandler(TUniquePtr<UE::MovieScene::IDeferredSignedObjectChangeHandler>&& InHandler)
+{
+	UE::MovieScene::GDeferredSignedObjectChangeHandler = MoveTemp(InHandler);
+}
+
+void UMovieSceneSignedObject::AddFlushSignal(TWeakPtr<UE::MovieScene::IDeferredSignedObjectFlushSignal> Signal)
+{
+	UE::MovieScene::GDeferredSignedObjectFlushSignals.Add(Signal);
+}
+
+void UMovieSceneSignedObject::ResetImplicitScopedModifyDefer()
+{
+	using namespace UE::MovieScene;
+	if (GDeferredSignedObjectChangeHandler)
+	{
+		GDeferredSignedObjectChangeHandler->ResetImplicitScopedModifyDefer();
+	}
+}
+
 void UMovieSceneSignedObject::PostInitProperties()
 {
 	Super::PostInitProperties();
@@ -31,18 +109,56 @@
 {
 	using namespace UE::MovieScene;
 
+	// We always change the signature immediately to ensure that any external code that wants
+	// to directly check our signature (eg, to clear caches) can still do so even while there
+	// is an outstanding signal pending for this object
 	Signature = FGuid::NewGuid();
 
+	// Regenerate the signature for all parents of this class
+	UObject* Outer = GetOuter();
+	while (Outer)
+	{
+		UMovieSceneSignedObject* TypedOuter = Cast<UMovieSceneSignedObject>(Outer);
+		if (TypedOuter)
+		{
+			TypedOuter->Signature = FGuid::NewGuid();
+		}
+		Outer = Outer->GetOuter();
+	}
+
+	// Give the change handler an opportunity to create an implicit scope if there's
+	// no explicit one active.
+	if (GScopedSignedObjectDeferCount == 0 && GDeferredSignedObjectChangeHandler)
+	{
+		GDeferredSignedObjectChangeHandler->CreateImplicitScopedModifyDefer();
+	}
+
+	if (GScopedSignedObjectDeferCount == 0 || !GDeferredSignedObjectChangeHandler)
+	{
+		BroadcastChanged();
+		SignalScopedSignedObjectModifyFlush();
+	}
+	else
+	{
+		GDeferredSignedObjectChangeHandler->DeferMarkAsChanged(this);
+	}
+}
+
+void UMovieSceneSignedObject::BroadcastChanged()
+{
+	using namespace UE::MovieScene;
+
+	EventHandlers.Trigger(&ISignedObjectEventHandler::OnModifiedDirectly, this);
 	OnSignatureChangedEvent.Broadcast();
-	
+
 	UObject* Outer = GetOuter();
 	while (Outer)
 	{
 		UMovieSceneSignedObject* TypedOuter = Cast<UMovieSceneSignedObject>(Outer);
 		if (TypedOuter)
 		{
-			TypedOuter->MarkAsChanged();
-			break;
+			TypedOuter->EventHandlers.Trigger(&ISignedObjectEventHandler::OnModifiedIndirectly, this);
+			TypedOuter->OnSignatureChangedEvent.Broadcast();
 		}
 		Outer = Outer->GetOuter();
 	}
@@ -51,6 +167,8 @@
 #if WITH_EDITOR
 bool UMovieSceneSignedObject::Modify(bool bAlwaysMarkDirty)
 {
+	using namespace UE::MovieScene;
+
 	bool bModified = Super::Modify(bAlwaysMarkDirty);
 	if ( bAlwaysMarkDirty )
 	{
@@ -67,14 +185,22 @@
 
 void UMovieSceneSignedObject::PostEditUndo()
 {
+	using namespace UE::MovieScene;
+
 	Super::PostEditUndo();
 	MarkAsChanged();
+
+	EventHandlers.Trigger(&ISignedObjectEventHandler::OnPostUndo);
 }
 
 void UMovieSceneSignedObject::PostEditUndo(TSharedPtr<ITransactionObjectAnnotation> TransactionAnnotation)
 {
+	using namespace UE::MovieScene;
+
 	Super::PostEditUndo(TransactionAnnotation);
 	MarkAsChanged();
+
+	EventHandlers.Trigger(&ISignedObjectEventHandler::OnPostUndo);
 }
 #endif
 
diff --git a/Engine/Source/Runtime/MovieScene/Public/Evaluation/MovieSceneEvaluationTrack.h b/Engine/Source/Runtime/MovieScene/Public/Evaluation/MovieSceneEvaluationTrack.h
--- a/Engine/Source/Runtime/MovieScene/Public/Evaluation/MovieSceneEvaluationTrack.h
+++ b/Engine/Source/Runtime/MovieScene/Public/Evaluation/MovieSceneEvaluationTrack.h
@@ -335,6 +335,11 @@
 	 */
 	MOVIESCENE_API void PostSerialize(const FArchive& Ar);
 
+	bool HasCustomInitialize() const
+	{
+		return TrackTemplate.IsValid() && TrackTemplate->HasCustomInitialize();
+	}
+
 public:
 
 	/**
diff --git a/Engine/Source/Runtime/MovieScene/Public/MovieScene.h b/Engine/Source/Runtime/MovieScene/Public/MovieScene.h
--- a/Engine/Source/Runtime/MovieScene/Public/MovieScene.h
+++ b/Engine/Source/Runtime/MovieScene/Public/MovieScene.h
@@ -16,6 +16,7 @@
 #include "MovieSceneObjectBindingID.h"
 #include "MovieSceneFrameMigration.h"
 #include "MovieSceneTimeController.h"
+#include "EventHandlers/MovieSceneDataEventContainer.h"
 #include "MovieScene.generated.h"
 
 struct FMovieSceneTimeController;
@@ -358,6 +359,8 @@
 
 public:
 
+	UE::MovieScene::TDataEventContainer<UE::MovieScene::ISequenceDataEventHandler> EventHandlers;
+
 	/**
 	 * Add a spawnable to this movie scene's list of owned blueprints.
 	 *
@@ -878,9 +881,44 @@
 	/**
 	 * Gets the root folders for this movie scene.
 	 */
-	TArray<UMovieSceneFolder*>& GetRootFolders();
+	TArrayView<UMovieSceneFolder* const> GetRootFolders();
 
 	/**
+	 * Gets a copy of the root folders for this movie scene.
+	 */
+	void GetRootFolders(TArray<UMovieSceneFolder*>& InRootFolders);
+
+	/**
+	 * Gets the number of root folders in this movie scene.
+	 */
+	int32 GetNumRootFolders() const;
+
+	/**
+	 * Gets the i-th root folder for this movie scene.
+	 */
+	UMovieSceneFolder* GetRootFolder(int32 FolderIndex) const;
+
+	/**
+	 * Adds a root folder for this movie scene.
+	 */
+	void AddRootFolder(UMovieSceneFolder* Folder);
+
+	/**
+	 * Removes a root folder for this movie scene (does not delete tracks or objects contained within)
+	 */
+	int32 RemoveRootFolder(UMovieSceneFolder* Folder);
+
+	/**
+	 * Removes a root folder for this movie scene (does not delete tracks or objects contained within)
+	 */
+	bool RemoveRootFolder(int32 FolderIndex);
+
+	/**
+	 * Removes all root folders from this movie scene (does not delete tracks or objects contained within)
+	 */
+	void EmptyRootFolders();
+
+	/**
 	 * Gets the nodes marked as solo in the editor, as node tree paths
 	 */
 	TArray<FString>& GetSoloNodes() { return SoloNodes; }
@@ -963,6 +1001,10 @@
 	{
 		return EditorData;
 	}
+	const FMovieSceneEditorData& GetEditorData() const
+	{
+		return EditorData;
+	}
 
 	void SetEditorData(FMovieSceneEditorData& InEditorData)
 	{
diff --git a/Engine/Source/Runtime/MovieScene/Public/MovieSceneBinding.h b/Engine/Source/Runtime/MovieScene/Public/MovieSceneBinding.h
--- a/Engine/Source/Runtime/MovieScene/Public/MovieSceneBinding.h
+++ b/Engine/Source/Runtime/MovieScene/Public/MovieSceneBinding.h
@@ -2,11 +2,12 @@
 
 #pragma once
 
-#include "CoreMinimal.h"
+#include "CoreTypes.h"
 #include "UObject/ObjectMacros.h"
 #include "Misc/Guid.h"
 #include "MovieSceneBinding.generated.h"
 
+class UMovieScene;
 class UMovieSceneTrack;
 
 /**
@@ -93,7 +94,7 @@
 	 *
 	 * @param NewTrack	The track to add
 	 */
-	void MOVIESCENE_API AddTrack(UMovieSceneTrack& NewTrack);
+	MOVIESCENE_API void AddTrack(UMovieSceneTrack& NewTrack, UMovieScene* Owner);
 
 	/**
 	 * Removes a track from this binding
@@ -101,9 +102,14 @@
 	 * @param Track	The track to remove
 	 * @return true if the track was successfully removed, false if the track could not be found
 	 */
-	bool RemoveTrack(UMovieSceneTrack& Track);
+	MOVIESCENE_API bool RemoveTrack(UMovieSceneTrack& Track, UMovieScene* Owner);
 
 	/**
+	 * Removes all null tracks from this binding
+	 */
+	MOVIESCENE_API void RemoveNullTracks();
+
+	/**
 	 * @return All tracks in this binding
 	 */
 	const TArray<UMovieSceneTrack*>& GetTracks() const
@@ -114,28 +120,12 @@
 	/**
 	 * Reset all tracks in this binding, returning the previous array of tracks
 	 */
-	TArray<UMovieSceneTrack*> StealTracks()
-	{
-		decltype(Tracks) Empty;
-		Swap(Empty, Tracks);
-		return Empty;
-	}
+	MOVIESCENE_API TArray<UMovieSceneTrack*> StealTracks(UMovieScene* Owner);
 
 	/**
 	 * Assign all tracks in this binding
 	 */
-	void SetTracks(TArray<UMovieSceneTrack*>&& InTracks)
-	{
-		Tracks = MoveTemp(InTracks);
-	}
-
-	/**
-	* Assign all tracks in this binding
-	*/
-	void SetTracks(TArray<UMovieSceneTrack*>& InTracks)
-	{
-		Tracks = InTracks;
-	}
+	MOVIESCENE_API void SetTracks(TArray<UMovieSceneTrack*>&& InTracks, UMovieScene* Owner);
 
 #if WITH_EDITORONLY_DATA
 	/**
diff --git a/Engine/Source/Runtime/MovieScene/Public/MovieSceneFolder.h b/Engine/Source/Runtime/MovieScene/Public/MovieSceneFolder.h
--- a/Engine/Source/Runtime/MovieScene/Public/MovieSceneFolder.h
+++ b/Engine/Source/Runtime/MovieScene/Public/MovieSceneFolder.h
@@ -7,6 +7,8 @@
 #include "UObject/ObjectMacros.h"
 #include "UObject/Object.h"
 #include "Misc/Guid.h"
+#include "EventHandlers/MovieSceneDataEventContainer.h"
+#include "EventHandlers/IFolderEventHandler.h"
 #include "MovieSceneFolder.generated.h"
 
 class UMovieSceneTrack;
@@ -24,7 +26,7 @@
 	void SetFolderName( FName InFolderName );
 
 	/** Gets the folders contained by this folder. */
-	const TArray<UMovieSceneFolder*>& GetChildFolders() const;
+	TArrayView<UMovieSceneFolder* const> GetChildFolders() const;
 
 	/** Adds a child folder to this folder. Automatically calls Modify() on the folder object. */
 	void AddChildFolder( UMovieSceneFolder* InChildFolder );
@@ -63,13 +65,17 @@
 	UMovieSceneFolder* FindFolderContaining(const FGuid& InObjectBinding);
 
 	/** Get the folder path for this folder, stopping at the given root folders */
-	static void CalculateFolderPath(UMovieSceneFolder* InFolder, const TArray<UMovieSceneFolder*>& RootFolders, TArray<FName>& FolderPath);
+	static void CalculateFolderPath(UMovieSceneFolder* InFolder, TArrayView<UMovieSceneFolder* const> RootFolders, TArray<FName>& FolderPath);
 	
 	/** For the given set of folders, return the folder that has the matching folder path */
-	static UMovieSceneFolder* GetFolderWithPath(const TArray<FName>& InFolderPath, const TArray<UMovieSceneFolder*>& InFolders, const TArray<UMovieSceneFolder*>& RootFolders);
+	static UMovieSceneFolder* GetFolderWithPath(const TArray<FName>& InFolderPath, const TArray<UMovieSceneFolder*>& InFolders, TArrayView<UMovieSceneFolder* const> RootFolders);
 
 	virtual void Serialize( FArchive& Archive );
 
+	FName MakeUniqueChildFolderName(FName InName) const;
+
+	static FName MakeUniqueChildFolderName(FName InName, TArrayView<UMovieSceneFolder* const> InFolders);
+
 #if WITH_EDITORONLY_DATA
 	/**
 	 * Get this folder's color.
@@ -110,6 +116,14 @@
 	}
 #endif
 
+#if WITH_EDITOR
+	virtual void PostEditUndo() override;
+	virtual void PostEditUndo(TSharedPtr<ITransactionObjectAnnotation> TransactionAnnotation) override;
+#endif
+
+	/** Event handlers for handling changes to this object */
+	UE::MovieScene::TDataEventContainer<UE::MovieScene::IFolderEventHandler> EventHandlers;
+
 private:
 	/** The name of this folder. */
 	UPROPERTY()
@@ -141,3 +155,6 @@
 	TArray<FGuid> ChildObjectBindings;
 };
 
+MOVIESCENE_API void GetMovieSceneFoldersRecursive(TArrayView<UMovieSceneFolder* const> InFoldersToRecurse, TArray<UMovieSceneFolder*>& OutFolders);
+
+
diff --git a/Engine/Source/Runtime/MovieScene/Public/MovieSceneFwd.h b/Engine/Source/Runtime/MovieScene/Public/MovieSceneFwd.h
--- a/Engine/Source/Runtime/MovieScene/Public/MovieSceneFwd.h
+++ b/Engine/Source/Runtime/MovieScene/Public/MovieSceneFwd.h
@@ -9,6 +9,20 @@
 #include "MovieSceneSequenceID.h"
 #include "HAL/PreprocessorHelpers.h"
 
+#ifndef UE_MOVIESCENE_EVENTS
+	#define UE_MOVIESCENE_EVENTS WITH_EDITOR
+#endif
+
+namespace UE
+{
+namespace MovieScene
+{
+
+class ISequenceDataEventHandler;
+
+} // namespace MovieScene
+} // namespace UE
+
 UENUM()
 namespace EMovieScenePlayerStatus
 {
diff --git a/Engine/Source/Runtime/MovieScene/Public/MovieScenePossessable.h b/Engine/Source/Runtime/MovieScene/Public/MovieScenePossessable.h
--- a/Engine/Source/Runtime/MovieScene/Public/MovieScenePossessable.h
+++ b/Engine/Source/Runtime/MovieScene/Public/MovieScenePossessable.h
@@ -14,6 +14,7 @@
 
 #include "MovieScenePossessable.generated.h"
 
+class UMovieScene;
 class IMovieScenePlayer;
 
 /**
@@ -124,15 +125,15 @@
 		return ParentGuid;
 	}
 
+	UE_DEPRECATED(5.1, "Please use the overload that takes a UMovieScene* to ensure that events are triggered correctly")
+	MOVIESCENE_API void SetParent(const FGuid& InParentGuid);
+
 	/**
 	 * Set the guid of this possessable's parent
 	 *
 	 * @param InParentGuid The guid of this possessable's parent.
 	 */
-	void SetParent(const FGuid& InParentGuid)
-	{
-		ParentGuid = InParentGuid;
-	}
+	MOVIESCENE_API void SetParent(const FGuid& InParentGuid, UMovieScene* Owner);
 
 	/** Array of tags that can be used for grouping and categorizing. */
 	UPROPERTY(EditAnywhere, AdvancedDisplay, Category=Actor)
diff --git a/Engine/Source/Runtime/MovieScene/Public/MovieSceneSection.h b/Engine/Source/Runtime/MovieScene/Public/MovieSceneSection.h
--- a/Engine/Source/Runtime/MovieScene/Public/MovieSceneSection.h
+++ b/Engine/Source/Runtime/MovieScene/Public/MovieSceneSection.h
@@ -10,6 +10,8 @@
 #include "MovieSceneFrameMigration.h"
 #include "MovieSceneSignedObject.h"
 #include "Evaluation/Blending/MovieSceneBlendType.h"
+#include "EventHandlers/ISectionEventHandler.h"
+#include "EventHandlers/MovieSceneDataEventContainer.h"
 
 #if UE_ENABLE_INCLUDE_ORDER_DEPRECATED_IN_5_1
 #include "Misc/FrameTime.h"
@@ -206,6 +208,8 @@
 	UPROPERTY(EditAnywhere, Category="Section", meta=(ShowOnlyInnerProperties))
 	FMovieSceneSectionEvalOptions EvalOptions;
 
+	UE::MovieScene::TDataEventContainer<UE::MovieScene::ISectionEventHandler> EventHandlers;
+
 public:
 
 	/**
@@ -463,7 +467,7 @@
 
 	/** Sets this section's new row index */
 	UFUNCTION(BlueprintCallable, Category = "Sequencer|Section")
-	void SetRowIndex(int32 NewRowIndex) {RowIndex = NewRowIndex;}
+	MOVIESCENE_API void SetRowIndex(int32 NewRowIndex);
 
 	/** Gets the row index for this section */
 	UFUNCTION(BlueprintPure, Category = "Sequencer|Section")
diff --git a/Engine/Source/Runtime/MovieScene/Public/MovieSceneSignedObject.h b/Engine/Source/Runtime/MovieScene/Public/MovieSceneSignedObject.h
--- a/Engine/Source/Runtime/MovieScene/Public/MovieSceneSignedObject.h
+++ b/Engine/Source/Runtime/MovieScene/Public/MovieSceneSignedObject.h
@@ -6,10 +6,55 @@
 #include "UObject/ObjectMacros.h"
 #include "UObject/Object.h"
 #include "Misc/Guid.h"
+#include "Templates/UniquePtr.h"
+#include "EventHandlers/MovieSceneDataEventContainer.h"
+#include "EventHandlers/ISignedObjectEventHandler.h"
+#include "Delegates/Delegate.h"
 #include "MovieSceneSignedObject.generated.h"
 
+class UMovieSceneSignedObject;
 class ITransactionObjectAnnotation;
 
+namespace UE
+{
+namespace MovieScene
+{
+
+struct IDeferredSignedObjectFlushSignal
+{
+	virtual ~IDeferredSignedObjectFlushSignal(){}
+
+	virtual void OnDeferredModifyFlush() = 0;
+};
+
+/**
+ * Application-wide utility interface that allows for deferral of UMovieSceneSignedObject::MarkAsChanged calls
+ */
+struct IDeferredSignedObjectChangeHandler
+{
+	virtual ~IDeferredSignedObjectChangeHandler(){}
+
+	virtual void Flush() = 0;
+	virtual void DeferMarkAsChanged(UMovieSceneSignedObject* SignedObject) = 0;
+	virtual bool CreateImplicitScopedModifyDefer() = 0;
+	virtual void ResetImplicitScopedModifyDefer() = 0;
+};
+
+struct MOVIESCENE_API FScopedSignedObjectModifyDefer
+{
+	FScopedSignedObjectModifyDefer(bool bInForceFlush = false);
+	~FScopedSignedObjectModifyDefer();
+
+private:
+	bool bForceFlush;
+};
+
+
+} // namespace MovieScene
+} // namespace UE
+
+
+
 UCLASS()
 class MOVIESCENE_API UMovieSceneSignedObject : public UObject
 {
@@ -18,23 +63,38 @@
 
 	UMovieSceneSignedObject(const FObjectInitializer& Init);
 
+	static void SetDeferredHandler(TUniquePtr<UE::MovieScene::IDeferredSignedObjectChangeHandler>&& InHandler);
+	static void AddFlushSignal(TWeakPtr<UE::MovieScene::IDeferredSignedObjectFlushSignal> Signal);
+	static void ResetImplicitScopedModifyDefer();
+
 	/**
-	 * 
+	 * Mark this object as having been changed in any way. Will regenerate this object's Signature
+	 * and schedule an event to be broadcast to notify subscribers of the change
 	 */
 	void MarkAsChanged();
 
 	/**
-	 * 
+	 * Retrieve this object's signature that uniquely identifies its current state.
+	 * Any change to this object will result in a new signature. This is a GUID and
+	 * not a hash - equivalent class state does not generate identical signatures.
 	 */
 	const FGuid& GetSignature() const
 	{
 		return Signature;
 	}
 
+	/**
+	 * Immediately broadcast events for this object being changed
+	 */
+	void BroadcastChanged();
+
 	/** Event that is triggered whenever this object's signature has changed */
 	DECLARE_EVENT(UMovieSceneSignedObject, FOnSignatureChanged)
 	FOnSignatureChanged& OnSignatureChanged() { return OnSignatureChangedEvent; }
 
+	/** Event that is triggered whenever this object's signature has changed */
+	UE::MovieScene::TDataEventContainer<UE::MovieScene::ISignedObjectEventHandler> EventHandlers;
+
 public:
 
 	virtual void PostInitProperties() override;
diff --git a/Engine/Source/Runtime/MovieSceneTracks/Private/Sections/MovieSceneAudioSection.cpp b/Engine/Source/Runtime/MovieSceneTracks/Private/Sections/MovieSceneAudioSection.cpp
--- a/Engine/Source/Runtime/MovieSceneTracks/Private/Sections/MovieSceneAudioSection.cpp
+++ b/Engine/Source/Runtime/MovieSceneTracks/Private/Sections/MovieSceneAudioSection.cpp
@@ -67,7 +67,7 @@
 #if WITH_EDITOR
 			FMovieSceneChannelMetaData Data;
 			FText TextName = FText::FromName(InName);	
-			Data.SetIdentifiers(InName, TextName, TextName);
+			Data.SetIdentifiers(FName(InName.ToString() + GetGeneratedTypeName<ChannelType>()), TextName, TextName);
 			InChannelProxyData.Add(const_cast<ChannelType&>(InChannel), Data, TMovieSceneExternalValue<ValueType>::Make());
 #else //WITH_EDITOR
 			InChannelProxyData.Add(const_cast<ChannelType&>(InChannel));
diff --git a/Engine/Source/Runtime/Slate/Private/Widgets/Views/STableViewBase.cpp b/Engine/Source/Runtime/Slate/Private/Widgets/Views/STableViewBase.cpp
--- a/Engine/Source/Runtime/Slate/Private/Widgets/Views/STableViewBase.cpp
+++ b/Engine/Source/Runtime/Slate/Private/Widgets/Views/STableViewBase.cpp
@@ -942,6 +942,11 @@
 	ItemsPanel->ClearItems();
 }
 
+const FChildren* STableViewBase::GetConstructedTableItems() const
+{
+	return ItemsPanel->GetChildren();
+}
+
 void STableViewBase::InsertPinnedWidget( const TSharedRef<SWidget> & WidgetToInset )
 {
 	PinnedItemsPanel->AddSlot(0)
diff --git a/Engine/Source/Runtime/Slate/Public/Framework/Views/TableViewTypeTraits.h b/Engine/Source/Runtime/Slate/Public/Framework/Views/TableViewTypeTraits.h
--- a/Engine/Source/Runtime/Slate/Public/Framework/Views/TableViewTypeTraits.h
+++ b/Engine/Source/Runtime/Slate/Public/Framework/Views/TableViewTypeTraits.h
@@ -54,6 +54,14 @@
 		Value = true
 	};
 };
+template <typename T, ESPMode Mode>
+struct TIsValidListItem<TWeakPtr<T, Mode>>
+{
+	enum
+	{
+		Value = true
+	};
+};
 template <typename T>
 struct TIsValidListItem<T*, typename TEnableIf<TPointerIsConvertibleFromTo<T, UObjectBase>::Value>::Type>
 {
@@ -293,6 +301,96 @@
 	class SerializerType{};
 };
 
+template <typename T> struct TListTypeTraits< TWeakPtr<T, ESPMode::NotThreadSafe> >
+{
+public:
+	typedef TWeakPtr<T> NullableType;
+
+	using MapKeyFuncs       = TDefaultMapHashableKeyFuncs<TWeakPtr<T, ESPMode::NotThreadSafe>, TSharedRef<ITableRow>, false>;
+	using MapKeyFuncsSparse = TDefaultMapHashableKeyFuncs<TWeakPtr<T, ESPMode::NotThreadSafe>, FSparseItemInfo, false>;
+	using SetKeyFuncs       = DefaultKeyFuncs<TWeakPtr<T, ESPMode::NotThreadSafe>>;
+
+	template<typename U>
+	static void AddReferencedObjects( FReferenceCollector&, 
+		TArray< TWeakPtr<T> >&, 
+		TSet< TWeakPtr<T> >&, 
+		TMap< const U*, TWeakPtr<T> >& )
+	{
+	}
+
+	static bool IsPtrValid( const TWeakPtr<T>& InPtr )
+	{
+		return InPtr.Pin().IsValid();
+	}
+
+	static void ResetPtr( TWeakPtr<T>& InPtr )
+	{
+		InPtr = nullptr;
+	}
+
+	static TWeakPtr<T> MakeNullPtr()
+	{
+		return TWeakPtr<T>();
+	}
+
+	static TWeakPtr<T> NullableItemTypeConvertToItemType( const TWeakPtr<T>& InPtr )
+	{
+		return InPtr;
+	}
+
+	static FString DebugDump( TWeakPtr<T> InPtr )
+	{
+		return InPtr.IsValid() ? FString::Printf(TEXT("0x%08x"), InPtr.Get()) : FString(TEXT("nullptr"));
+	}
+
+	class SerializerType{};
+};
+
+
+template <typename T> struct TListTypeTraits< TWeakPtr<T, ESPMode::ThreadSafe> >
+{
+public:
+	typedef TWeakPtr<T, ESPMode::ThreadSafe> NullableType;
+
+	using MapKeyFuncs       = TDefaultMapHashableKeyFuncs<TWeakPtr<T, ESPMode::ThreadSafe>, TSharedRef<ITableRow>, false>;
+	using MapKeyFuncsSparse = TDefaultMapHashableKeyFuncs<TWeakPtr<T, ESPMode::ThreadSafe>, FSparseItemInfo, false>;
+	using SetKeyFuncs       = DefaultKeyFuncs<TWeakPtr<T, ESPMode::ThreadSafe>>;
+
+	template<typename U>
+	static void AddReferencedObjects( FReferenceCollector&, 
+		TArray< TWeakPtr<T, ESPMode::ThreadSafe> >&, 
+		TSet< TWeakPtr<T, ESPMode::ThreadSafe> >&, 
+		TMap< const U*, TWeakPtr<T, ESPMode::ThreadSafe> >& WidgetToItemMap)
+	{
+	}
+
+	static bool IsPtrValid( const TWeakPtr<T, ESPMode::ThreadSafe>& InPtr )
+	{
+		return InPtr.Pin().IsValid();
+	}
+
+	static void ResetPtr( TWeakPtr<T, ESPMode::ThreadSafe>& InPtr )
+	{
+		InPtr = nullptr;
+	}
+
+	static TWeakPtr<T, ESPMode::ThreadSafe> MakeNullPtr()
+	{
+		return TWeakPtr<T, ESPMode::ThreadSafe>(nullptr);
+	}
+
+	static TWeakPtr<T, ESPMode::ThreadSafe> NullableItemTypeConvertToItemType( const TWeakPtr<T, ESPMode::ThreadSafe>& InPtr )
+	{
+		return InPtr;
+	}
+
+	static FString DebugDump( TWeakPtr<T, ESPMode::ThreadSafe> InPtr )
+	{
+		return InPtr.IsValid() ? FString::Printf(TEXT("0x%08x"), InPtr.Pin().Get()) : FString(TEXT("nullptr"));
+	}
+
+	class SerializerType{};
+};
 
 /**
  * Pointer-related functionality (e.g. setting to nullptr, testing for nullptr) specialized for SharedPointers.
diff --git a/Engine/Source/Runtime/Slate/Public/Widgets/Views/STableRow.h b/Engine/Source/Runtime/Slate/Public/Widgets/Views/STableRow.h
--- a/Engine/Source/Runtime/Slate/Public/Widgets/Views/STableRow.h
+++ b/Engine/Source/Runtime/Slate/Public/Widgets/Views/STableRow.h
@@ -84,7 +84,8 @@
 	DECLARE_DELEGATE_RetVal_ThreeParams(TOptional<EItemDropZone>, FOnCanAcceptDrop, const FDragDropEvent&, EItemDropZone, ItemType);
 	/** Delegate signature for handling the drop of FDragDropEvent onto target of type ItemType */
 	DECLARE_DELEGATE_RetVal_ThreeParams(FReply, FOnAcceptDrop, const FDragDropEvent&, EItemDropZone, ItemType);
-
+	/** Delegate signature for painting drop indicators. */
+	DECLARE_DELEGATE_RetVal_EightParams(int32, FOnPaintDropIndicator, EItemDropZone, const FPaintArgs&, const FGeometry&, const FSlateRect&, FSlateWindowElementList&, int32, const FWidgetStyle&, bool);
 public:
 
 	SLATE_BEGIN_ARGS( STableRow< ItemType > )
@@ -120,6 +121,11 @@
 		 */
 		SLATE_EVENT( FOnAcceptDrop,    OnAcceptDrop )
 
+		/**
+		 * Used for painting drop indicators
+		 */
+		SLATE_EVENT( FOnPaintDropIndicator, OnPaintDropIndicator )
+
 		// Low level DragAndDrop
 		SLATE_EVENT( FOnDragDetected,      OnDragDetected )
 		SLATE_EVENT( FOnTableRowDragEnter, OnDragEnter )
@@ -303,6 +309,18 @@
 	}
 #endif
 
+	/** Retrieves a brush for rendering a drop indicator for the specified drop zone */
+	const FSlateBrush* GetDropIndicatorBrush(EItemDropZone InItemDropZone) const
+	{
+		switch (InItemDropZone)
+		{
+			case EItemDropZone::AboveItem: return &Style->DropIndicator_Above; break;
+			default:
+			case EItemDropZone::OntoItem: return &Style->DropIndicator_Onto; break;
+			case EItemDropZone::BelowItem: return &Style->DropIndicator_Below; break;
+		};
+	}
+
 	virtual int32 OnPaint( const FPaintArgs& Args, const FGeometry& AllottedGeometry, const FSlateRect& MyCullingRect, FSlateWindowElementList& OutDrawElements, int32 LayerId, const FWidgetStyle& InWidgetStyle, bool bParentEnabled ) const override
 	{
 		TSharedRef< ITypedTableView<ItemType> > OwnerTable = OwnerTablePtr.Pin().ToSharedRef();
@@ -327,55 +345,62 @@
 
 		if (ItemDropZone.IsSet())
 		{
-			// Draw feedback for user dropping an item above, below, or onto a row.
-			const FSlateBrush* DropIndicatorBrush = [&]()
+			if (PaintDropIndicatorEvent.IsBound())
 			{
-				switch (ItemDropZone.GetValue())
-				{
-					case EItemDropZone::AboveItem: return &Style->DropIndicator_Above; break;
-					default:
-					case EItemDropZone::OntoItem: return &Style->DropIndicator_Onto; break;
-					case EItemDropZone::BelowItem: return &Style->DropIndicator_Below; break;
-				};
-			}();
-
-			if (OwnerTable->Private_GetOrientation() == Orient_Vertical)
-			{
-				FSlateDrawElement::MakeBox
-				(
-					OutDrawElements,
-					LayerId++,
-					AllottedGeometry.ToPaintGeometry(),
-					DropIndicatorBrush,
-					ESlateDrawEffect::None,
-					DropIndicatorBrush->GetTint(InWidgetStyle) * InWidgetStyle.GetColorAndOpacityTint()
-				);
+				LayerId = PaintDropIndicatorEvent.Execute(ItemDropZone.GetValue(), Args, AllottedGeometry, MyCullingRect, OutDrawElements, LayerId, InWidgetStyle, bParentEnabled);
 			}
 			else
 			{
-				// Reuse the drop indicator asset for horizontal, by rotating the drawn box 90 degrees.
-				const FVector2D LocalSize(AllottedGeometry.GetLocalSize());
-				const FVector2D Pivot(LocalSize * 0.5f);
-				const FVector2D RotatedLocalSize(LocalSize.Y, LocalSize.X);
-				FSlateLayoutTransform RotatedTransform(Pivot - RotatedLocalSize * 0.5f);	// Make the box centered to the alloted geometry, so that it can be rotated around the center.
-
-				FSlateDrawElement::MakeRotatedBox(
-					OutDrawElements,
-					LayerId++,
-					AllottedGeometry.ToPaintGeometry(RotatedLocalSize, RotatedTransform),
-					DropIndicatorBrush,
-					ESlateDrawEffect::None,
-					-UE_HALF_PI,	// 90 deg CCW
-					RotatedLocalSize * 0.5f,	// Relative center to the flipped
-					FSlateDrawElement::RelativeToElement,
-					DropIndicatorBrush->GetTint(InWidgetStyle) * InWidgetStyle.GetColorAndOpacityTint()
-				);
+				OnPaintDropIndicator(ItemDropZone.GetValue(), Args, AllottedGeometry, MyCullingRect, OutDrawElements, LayerId, InWidgetStyle, bParentEnabled);
 			}
 		}
 
 		return LayerId;
 	}
 
+	virtual int32 OnPaintDropIndicator( EItemDropZone InItemDropZone, const FPaintArgs& Args, const FGeometry& AllottedGeometry, const FSlateRect& MyCullingRect, FSlateWindowElementList& OutDrawElements, int32 LayerId, const FWidgetStyle& InWidgetStyle, bool bParentEnabled ) const
+	{
+		TSharedRef< ITypedTableView<ItemType> > OwnerTable = OwnerTablePtr.Pin().ToSharedRef();
+
+		// Draw feedback for user dropping an item above, below, or onto a row.
+		const FSlateBrush* DropIndicatorBrush = GetDropIndicatorBrush(InItemDropZone);
+
+		if (OwnerTable->Private_GetOrientation() == Orient_Vertical)
+		{
+			FSlateDrawElement::MakeBox
+			(
+				OutDrawElements,
+				LayerId++,
+				AllottedGeometry.ToPaintGeometry(),
+				DropIndicatorBrush,
+				ESlateDrawEffect::None,
+				DropIndicatorBrush->GetTint(InWidgetStyle) * InWidgetStyle.GetColorAndOpacityTint()
+			);
+		}
+		else
+		{
+			// Reuse the drop indicator asset for horizontal, by rotating the drawn box 90 degrees.
+			const FVector2D LocalSize(AllottedGeometry.GetLocalSize());
+			const FVector2D Pivot(LocalSize * 0.5f);
+			const FVector2D RotatedLocalSize(LocalSize.Y, LocalSize.X);
+			FSlateLayoutTransform RotatedTransform(Pivot - RotatedLocalSize * 0.5f);	// Make the box centered to the alloted geometry, so that it can be rotated around the center.
+
+			FSlateDrawElement::MakeRotatedBox(
+				OutDrawElements,
+				LayerId++,
+				AllottedGeometry.ToPaintGeometry(RotatedLocalSize, RotatedTransform),
+				DropIndicatorBrush,
+				ESlateDrawEffect::None,
+				-UE_HALF_PI,	// 90 deg CCW
+				RotatedLocalSize * 0.5f,	// Relative center to the flipped
+				FSlateDrawElement::RelativeToElement,
+				DropIndicatorBrush->GetTint(InWidgetStyle) * InWidgetStyle.GetColorAndOpacityTint()
+			);
+		}
+
+		return LayerId;
+	}
+
 	/**
 	 * Called when a mouse button is double clicked.  Override this in derived classes.
 	 *
@@ -1198,6 +1223,9 @@
 	/** @see STableRow's OnAcceptDrop event */
 	FOnAcceptDrop OnAcceptDrop;
 
+	/** Optional delegate for painting drop indicators */
+	FOnPaintDropIndicator PaintDropIndicatorEvent;
+
 	/** Are we currently dragging/dropping over this item? */
 	TOptional<EItemDropZone> ItemDropZone;
 
diff --git a/Engine/Source/Runtime/Slate/Public/Widgets/Views/STableViewBase.h b/Engine/Source/Runtime/Slate/Public/Widgets/Views/STableViewBase.h
--- a/Engine/Source/Runtime/Slate/Public/Widgets/Views/STableViewBase.h
+++ b/Engine/Source/Runtime/Slate/Public/Widgets/Views/STableViewBase.h
@@ -246,7 +246,9 @@
 
 	/** Add a WidgetToAppend to the bottom of the view. */
 	void AppendWidget( const TSharedRef<ITableRow>& WidgetToAppend );
-	
+
+	const FChildren* GetConstructedTableItems() const;
+
 	/**
 	 * Remove all the widgets from the view.
 	 */
