Sequencer: Fixed hierarchy modification during iteration

diff --git a/Engine/Source/Editor/SequencerCore/Private/MVVM/ViewModels/ViewModelHierarchy.cpp b/Engine/Source/Editor/SequencerCore/Private/MVVM/ViewModels/ViewModelHierarchy.cpp
--- a/Engine/Source/Editor/SequencerCore/Private/MVVM/ViewModels/ViewModelHierarchy.cpp
+++ b/Engine/Source/Editor/SequencerCore/Private/MVVM/ViewModels/ViewModelHierarchy.cpp
@@ -346,7 +346,7 @@
 	{
 		FViewModelHierarchyOperation Operation(Owner);
 
-		for (const FViewModelPtr& Child : IterateSubList())
+		for (const FViewModelPtr& Child : IterateSubList().ToArray())
 		{
 			Child->SetParentOnly(nullptr);
 		}
diff --git a/Engine/Source/Editor/SequencerCore/Public/MVVM/ViewModels/ViewModelIterators.h b/Engine/Source/Editor/SequencerCore/Public/MVVM/ViewModels/ViewModelIterators.h
--- a/Engine/Source/Editor/SequencerCore/Public/MVVM/ViewModels/ViewModelIterators.h
+++ b/Engine/Source/Editor/SequencerCore/Public/MVVM/ViewModels/ViewModelIterators.h
@@ -73,6 +73,29 @@
 	FORCEINLINE FViewModelSubListIterator begin() const { return *this; }
 	FORCEINLINE FViewModelSubListIterator end() const   { return FViewModelSubListIterator{ nullptr }; }
 
+	/**
+	 * Populate the specified array with the remaining items in this iterator
+	 * The state of this iterator remains unchanged
+	 */
+	void ToArray(TArray<FViewModelPtr>& OutArray) const
+	{
+		for (FViewModelSubListIterator Temp = *this; Temp; ++Temp)
+		{
+			OutArray.Add(*Temp);
+		}
+	}
+
+	/**
+	 * Return the remaining items in this iterator as an array
+	 * The state of this iterator remains unchanged
+	 */
+	TArray<FViewModelPtr> ToArray() const
+	{
+		TArray<FViewModelPtr> Temp;
+		ToArray(Temp);
+		return Temp;
+	}
+
 private:
 
 	/** The pointed-to data model */
