Sequencer: Fixed mutation during iteration

diff --git a/Engine/Source/Editor/SequencerCore/Private/MVVM/ViewModels/ViewModelHierarchy.cpp b/Engine/Source/Editor/SequencerCore/Private/MVVM/ViewModels/ViewModelHierarchy.cpp
--- a/Engine/Source/Editor/SequencerCore/Private/MVVM/ViewModels/ViewModelHierarchy.cpp
+++ b/Engine/Source/Editor/SequencerCore/Private/MVVM/ViewModels/ViewModelHierarchy.cpp
@@ -296,7 +296,7 @@ void FViewModelChildren::MoveChildrenTo(const FViewModelChildren& OutDestination
 		// Already same parent so we're just moving from one list to another
 		ListTail = GetTail();
 	}
-	else for (const FViewModelPtr& Child : IterateSubList())
+	else for (const FViewModelPtr& Child : IterateSubList().ToArray())
 	{
 		// We purposefully only report changes for the first element
 		// Since that is all that is necessary to trigger events for the old parent
