Don't include filtered-out outliner items when computing vertical track area lane positions.

diff --git a/Engine/Source/Editor/SequencerCore/Private/MVVM/Views/STrackLane.cpp b/Engine/Source/Editor/SequencerCore/Private/MVVM/Views/STrackLane.cpp
--- a/Engine/Source/Editor/SequencerCore/Private/MVVM/Views/STrackLane.cpp
+++ b/Engine/Source/Editor/SequencerCore/Private/MVVM/Views/STrackLane.cpp
@@ -360,9 +360,15 @@
 
 			if (IOutlinerExtension* OutlinerExtension = It->CastThis<IOutlinerExtension>())
 			{
-				AccumulatedItemTop -= OutlinerExtension->GetOutlinerSizing().GetTotalHeight();
+				const bool bIsExpanded = OutlinerExtension->IsExpanded();
+				const bool bIsFilteredOut = OutlinerExtension->IsFilteredOut();
 
-				if (!OutlinerExtension->IsExpanded())
+				if (!bIsFilteredOut)
+				{
+					AccumulatedItemTop -= OutlinerExtension->GetOutlinerSizing().GetTotalHeight();
+				}
+
+				if (bIsFilteredOut || !bIsExpanded)
 				{
 					It.IgnoreCurrentChildren();
 				}
