Sequencer: Fix crash in FViewModelExtensionCollection::Destroy

diff --git a/Engine/Source/Editor/SequencerCore/Private/MVVM/Extensions/ViewModelExtensionCollection.cpp b/Engine/Source/Editor/SequencerCore/Private/MVVM/Extensions/ViewModelExtensionCollection.cpp
--- a/Engine/Source/Editor/SequencerCore/Private/MVVM/Extensions/ViewModelExtensionCollection.cpp
+++ b/Engine/Source/Editor/SequencerCore/Private/MVVM/Extensions/ViewModelExtensionCollection.cpp
@@ -94,7 +94,10 @@
 	if (TSharedPtr<FViewModel> Model = WeakModel.Pin())
 	{
 		TSharedPtr<FSharedViewModelData> SharedData = Model->GetSharedData();
-		SharedData->UnsubscribeFromHierarchyChanged(Model, OnHierarchyUpdatedHandle);
+		if (SharedData)
+		{
+			SharedData->UnsubscribeFromHierarchyChanged(Model, OnHierarchyUpdatedHandle);
+		}
 	}
 
 	OnHierarchyUpdatedHandle = FDelegateHandle();
