Sequencer: Rename master/slave views to primary/pinned views

diff --git a/Engine/Source/Editor/Sequencer/Private/SSequencer.cpp b/Engine/Source/Editor/Sequencer/Private/SSequencer.cpp
--- a/Engine/Source/Editor/Sequencer/Private/SSequencer.cpp
+++ b/Engine/Source/Editor/Sequencer/Private/SSequencer.cpp
@@ -273,7 +273,7 @@ void SSequencer::Construct(const FArguments& InArgs, TSharedRef<FSequencer> InSe
 
 	PinnedTrackArea->SetOutliner(PinnedTreeView);
 	PinnedTrackArea->SetShowPinned(true);
-	PinnedTrackArea->SetIsSlave(true);
+	PinnedTrackArea->SetIsPinned(true);
 	PinnedTreeView->SetShowPinned(true);
 
 	SAssignNew(TrackArea, SSequencerTrackAreaView, InSequencer->GetViewModel()->GetTrackArea(), TimeSliderControllerRef);
@@ -284,7 +284,7 @@ void SSequencer::Construct(const FArguments& InArgs, TSharedRef<FSequencer> InSe
 
 	TrackArea->SetOutliner(TreeView);
 
-	TreeView->AddSlaveTreeView(PinnedTreeView);
+	TreeView->AddPinnedTreeView(PinnedTreeView);
 
 	TAttribute<FAnimatedRange> ViewRangeAttribute = InArgs._ViewRange;
 
diff --git a/Engine/Source/Editor/SequencerCore/Private/MVVM/Views/SOutlinerView.cpp b/Engine/Source/Editor/SequencerCore/Private/MVVM/Views/SOutlinerView.cpp
--- a/Engine/Source/Editor/SequencerCore/Private/MVVM/Views/SOutlinerView.cpp
+++ b/Engine/Source/Editor/SequencerCore/Private/MVVM/Views/SOutlinerView.cpp
@@ -256,7 +256,7 @@
 SOutlinerView::~SOutlinerView()
 {
 	TrackArea.Reset();
-	SlaveTreeViews.Empty();
+	PinnedTreeViews.Empty();
 }
 
 void SOutlinerView::Tick(const FGeometry& AllottedGeometry, const double InCurrentTime, const float InDeltaTime)
@@ -465,10 +465,10 @@
 	}
 }
 
-void SOutlinerView::AddSlaveTreeView(TSharedPtr<SOutlinerView> SlaveTreeView)
+void SOutlinerView::AddPinnedTreeView(TSharedPtr<SOutlinerView> PinnedTreeView)
 {
-	SlaveTreeViews.Add(SlaveTreeView);
-	SlaveTreeView->SetMasterTreeView(SharedThis(this));
+	PinnedTreeViews.Add(PinnedTreeView);
+	PinnedTreeView->SetPrimaryTreeView(SharedThis(this));
 }
 
 void SOutlinerView::OnRightMouseButtonDown(const FPointerEvent& MouseEvent)
@@ -847,9 +847,9 @@
 	RebuildList();
 	//RequestTreeRefresh();
 
-	for (TSharedPtr<SOutlinerView> SlaveTreeView : SlaveTreeViews)
+	for (TSharedPtr<SOutlinerView> PinnedTreeView : PinnedTreeViews)
 	{
-		SlaveTreeView->Refresh();
+		PinnedTreeView->Refresh();
 	}
 }
 
diff --git a/Engine/Source/Editor/SequencerCore/Private/MVVM/Views/STrackAreaView.cpp b/Engine/Source/Editor/SequencerCore/Private/MVVM/Views/STrackAreaView.cpp
--- a/Engine/Source/Editor/SequencerCore/Private/MVVM/Views/STrackAreaView.cpp
+++ b/Engine/Source/Editor/SequencerCore/Private/MVVM/Views/STrackAreaView.cpp
@@ -426,7 +426,7 @@
 {
 	FVector2D Size = AllottedGeometry.GetLocalSize();
 
-	if (!IsSlave() && SizeLastFrame.IsSet() && Size.X != SizeLastFrame->X)
+	if (!IsPinned() && SizeLastFrame.IsSet() && Size.X != SizeLastFrame->X)
 	{
 		OnResized(SizeLastFrame.GetValue(), Size);
 	}
diff --git a/Engine/Source/Editor/SequencerCore/Public/MVVM/Views/SOutlinerView.h b/Engine/Source/Editor/SequencerCore/Public/MVVM/Views/SOutlinerView.h
--- a/Engine/Source/Editor/SequencerCore/Public/MVVM/Views/SOutlinerView.h
+++ b/Engine/Source/Editor/SequencerCore/Public/MVVM/Views/SOutlinerView.h
@@ -182,10 +182,10 @@
 	void UpdateTrackArea();
 
 	/** Add a SOutlinerView object that should be modified or updated when this Treeview is updated */
-	void AddSlaveTreeView(TSharedPtr<SOutlinerView> SlaveTreeView);
+	void AddPinnedTreeView(TSharedPtr<SOutlinerView> PinnedTreeView);
 
-	/** Set a SOutlinerView object this Treeview is slaved to, for operations that should happen on the master */
-	void SetMasterTreeView(TSharedPtr<SOutlinerView> InMasterTreeView) { MasterTreeView = InMasterTreeView; }
+	/** Set a SOutlinerView object this Treeview is pinned to, for operations that should happen on the primary */
+	void SetPrimaryTreeView(TSharedPtr<SOutlinerView> InPrimaryTreeView) { PrimaryTreeView = InPrimaryTreeView; }
 
 	/** Set whether this TreeView should show only pinned nodes or only non-pinned nodes  */
 	void SetShowPinned(bool bShowPinned) { bShowPinnedNodes = bShowPinned; }
@@ -224,10 +224,10 @@
 	TSharedPtr<STrackAreaView> TrackArea;
 
 	/** SOutlinerView objects that should be modified or updated when this Treeview is updated */
-	TArray<TSharedPtr<SOutlinerView>> SlaveTreeViews;
+	TArray<TSharedPtr<SOutlinerView>> PinnedTreeViews;
 
-	/** The SOutlinerView object this SOutlinerView is slave to, or nullptr if not a slave */
-	TWeakPtr<SOutlinerView> MasterTreeView;
+	/** The SOutlinerView object this SOutlinerView is pinned to, or nullptr if not pinned */
+	TWeakPtr<SOutlinerView> PrimaryTreeView;
 
 	float VirtualTop;
 
diff --git a/Engine/Source/Editor/SequencerCore/Public/MVVM/Views/STrackAreaView.h b/Engine/Source/Editor/SequencerCore/Public/MVVM/Views/STrackAreaView.h
--- a/Engine/Source/Editor/SequencerCore/Public/MVVM/Views/STrackAreaView.h
+++ b/Engine/Source/Editor/SequencerCore/Public/MVVM/Views/STrackAreaView.h
@@ -89,9 +89,9 @@
 	void SetShowPinned(bool bShowPinned) { bShowPinnedNodes = bShowPinned; }
 	bool ShowPinned() const { return bShowPinnedNodes; }
 
-	/** Set whether this TrackArea is slaved to a master TrackArea and should skip updating external controls */
-	void SetIsSlave(bool bInIsSlave) { bIsSlave = bInIsSlave; }
-	bool IsSlave() const { return bIsSlave; }
+	/** Set whether this TrackArea is pinned to another TrackArea and should skip updating external controls */
+	void SetIsPinned(bool bInIsPinned) { bIsPinned = bInIsPinned; }
+	bool IsPinned() const { return bIsPinned; }
 
 	static FLinearColor BlendDefaultTrackColor(FLinearColor InColor);
 
@@ -162,8 +162,8 @@
 	/** Whether this TrackArea is for pinned nodes or non-pinned nodes */
 	bool bShowPinnedNodes;
 
-	/** Whether this TrackArea is slaved to a master TrackArea and should skip updating external controls */
-	bool bIsSlave;
+	/** Whether this TrackArea is pinned to another TrackArea and should skip updating external controls */
+	bool bIsPinned;
 };
 
 } // namespace Sequencer
