Sequencer: Fixed extension collections not being updated on access when there was a latent hierarchy changed notification pending

diff --git a/Engine/Source/Editor/SequencerCore/Private/MVVM/Extensions/ViewModelExtensionCollection.cpp b/Engine/Source/Editor/SequencerCore/Private/MVVM/Extensions/ViewModelExtensionCollection.cpp
--- a/Engine/Source/Editor/SequencerCore/Private/MVVM/Extensions/ViewModelExtensionCollection.cpp
+++ b/Engine/Source/Editor/SequencerCore/Private/MVVM/Extensions/ViewModelExtensionCollection.cpp
@@ -31,7 +31,14 @@ FViewModelExtensionCollection::FViewModelExtensionCollection(FViewModelTypeID In
 
 FViewModelExtensionCollection::~FViewModelExtensionCollection()
 {
-	Destroy();
+	// Clean up our subscription to HierarchyChanged events. We do not call Destroy
+	// since that could call OnExtensionsDirtied which is a virtual function
+	TSharedPtr<FViewModel>           Model      = WeakModel.Pin();
+	TSharedPtr<FSharedViewModelData> SharedData = Model ? Model->GetSharedData() : nullptr;
+	if (SharedData && OnHierarchyUpdatedHandle.IsValid())
+	{
+		SharedData->UnsubscribeFromHierarchyChanged(Model, OnHierarchyUpdatedHandle);
+	}
 }
 
 void FViewModelExtensionCollection::Initialize()
@@ -60,12 +67,44 @@ void FViewModelExtensionCollection::Reinitialize(TWeakPtr<FViewModel> InWeakMode
 	Initialize();
 }
 
+void FViewModelExtensionCollection::ConditionalUpdate() const
+{
+	TSharedPtr<FViewModel> Model = WeakModel.Pin();
+
+	// If we have a valid mode and are subscribed to a HierarchyChanged notification
+	// make sure we report any outstanding operations that might affect the ptrs we cached
+	if (Model && OnHierarchyUpdatedHandle.IsValid())
+	{
+		TSharedPtr<FSharedViewModelData> SharedData = Model->GetSharedData();
+		if (SharedData)
+		{
+			SharedData->ReportLatentHierarchicalOperations();
+		}
+	}
+	else if (ExtensionContainer.Num() != 0)
+	{
+		// Our model is no longer valid or we're not subscribed to notifications,
+		// we cannot cache any ptrs because we can't keep them are up-to-date
+		bNeedsUpdate = false;
+		ExtensionContainer.Reset();
+	}
+	else if (bNeedsUpdate)
+	{
+		// Our pointers could be out of date so update them now
+		Update();
+	}
+}
+
 void FViewModelExtensionCollection::FViewModelExtensionCollection::Update() const
 {
 	bNeedsUpdate = false;
 	ExtensionContainer.Reset();
 
-	if (TSharedPtr<FViewModel> Model = WeakModel.Pin())
+	TSharedPtr<FViewModel> Model = WeakModel.Pin();
+
+	// Do not allow caching any ptrs from the tree unless we have a valid 
+	// OnHierarchyUpdatedHandle that guarantees we can remain up-to-date
+	if (Model && OnHierarchyUpdatedHandle.IsValid())
 	{
 		FParentFirstChildIterator ChildIt = Model->GetDescendants();
 		if (DesiredRecursionDepth != -1)
diff --git a/Engine/Source/Editor/SequencerCore/Public/MVVM/Extensions/ViewModelExtensionCollection.h b/Engine/Source/Editor/SequencerCore/Public/MVVM/Extensions/ViewModelExtensionCollection.h
--- a/Engine/Source/Editor/SequencerCore/Public/MVVM/Extensions/ViewModelExtensionCollection.h
+++ b/Engine/Source/Editor/SequencerCore/Public/MVVM/Extensions/ViewModelExtensionCollection.h
@@ -53,13 +53,13 @@
 
 	void OnHierarchyUpdated();
 
+	void ConditionalUpdate() const;
+
 	template<typename T>
 	TArrayView<T* const> GetExtensions() const
 	{
-		if (bNeedsUpdate)
-		{
-			Update();
-		}
+		// Conditionally update any ptrs before we allow any access to them
+		ConditionalUpdate();
 
 		void* const * BaseExtensions = ExtensionContainer.GetData();
 		T* const * Data = reinterpret_cast<T* const *>(BaseExtensions);
