Fix VM extension collection not updating correctly

This made the layer bars never properly computed.

diff --git a/Engine/Source/Editor/SequencerCore/Private/MVVM/Extensions/ViewModelExtensionCollection.cpp b/Engine/Source/Editor/SequencerCore/Private/MVVM/Extensions/ViewModelExtensionCollection.cpp
--- a/Engine/Source/Editor/SequencerCore/Private/MVVM/Extensions/ViewModelExtensionCollection.cpp
+++ b/Engine/Source/Editor/SequencerCore/Private/MVVM/Extensions/ViewModelExtensionCollection.cpp
@@ -33,12 +33,7 @@ FViewModelExtensionCollection::~FViewModelExtensionCollection()
 {
 	// Clean up our subscription to HierarchyChanged events. We do not call Destroy
 	// since that could call OnExtensionsDirtied which is a virtual function
-	TSharedPtr<FViewModel>           Model      = WeakModel.Pin();
-	TSharedPtr<FSharedViewModelData> SharedData = Model ? Model->GetSharedData() : nullptr;
-	if (SharedData && OnHierarchyUpdatedHandle.IsValid())
-	{
-		SharedData->UnsubscribeFromHierarchyChanged(Model, OnHierarchyUpdatedHandle);
-	}
+	DestroyImpl();
 }
 
 void FViewModelExtensionCollection::Initialize()
@@ -71,7 +66,7 @@ void FViewModelExtensionCollection::ConditionalUpdate() const
 {
 	TSharedPtr<FViewModel> Model = WeakModel.Pin();
 
-	// If we have a valid mode and are subscribed to a HierarchyChanged notification
+	// If we have a valid model and are subscribed to a HierarchyChanged notification
 	// make sure we report any outstanding operations that might affect the ptrs we cached
 	if (Model && OnHierarchyUpdatedHandle.IsValid())
 	{
@@ -88,7 +83,8 @@ void FViewModelExtensionCollection::ConditionalUpdate() const
 		bNeedsUpdate = false;
 		ExtensionContainer.Reset();
 	}
-	else if (bNeedsUpdate)
+	
+	if (bNeedsUpdate)
 	{
 		// Our pointers could be out of date so update them now
 		Update();
@@ -129,6 +125,14 @@ void FViewModelExtensionCollection::OnHierarchyUpdated()
 }
 
 void FViewModelExtensionCollection::Destroy()
+{
+	DestroyImpl();
+
+	bNeedsUpdate = true;
+	OnExtensionsDirtied();
+}
+
+void FViewModelExtensionCollection::DestroyImpl()
 {
 	if (TSharedPtr<FViewModel> Model = WeakModel.Pin())
 	{
@@ -140,9 +144,6 @@ void FViewModelExtensionCollection::Destroy()
 	}
 
 	OnHierarchyUpdatedHandle = FDelegateHandle();
-
-	bNeedsUpdate = true;
-	OnExtensionsDirtied();
 }
 
 } // namespace Sequencer
diff --git a/Engine/Source/Editor/SequencerCore/Public/MVVM/Extensions/ViewModelExtensionCollection.h b/Engine/Source/Editor/SequencerCore/Public/MVVM/Extensions/ViewModelExtensionCollection.h
--- a/Engine/Source/Editor/SequencerCore/Public/MVVM/Extensions/ViewModelExtensionCollection.h
+++ b/Engine/Source/Editor/SequencerCore/Public/MVVM/Extensions/ViewModelExtensionCollection.h
@@ -55,6 +55,8 @@ private:
 
 	void ConditionalUpdate() const;
 
+	void DestroyImpl();
+
 	template<typename T>
 	TArrayView<T* const> GetExtensions() const
 	{
