Sequencer: Stop key bars from being drawn when EKeyBarConnectionStyle == None

diff --git a/Engine/Source/Editor/SequencerCore/Private/MVVM/Views/KeyRenderer.cpp b/Engine/Source/Editor/SequencerCore/Private/MVVM/Views/KeyRenderer.cpp
--- a/Engine/Source/Editor/SequencerCore/Private/MVVM/Views/KeyRenderer.cpp
+++ b/Engine/Source/Editor/SequencerCore/Private/MVVM/Views/KeyRenderer.cpp
@@ -518,7 +518,8 @@ void FKeyRenderer::UpdateViewDependentData(const FKeyBatchParameters& Params, co
 		}
 		else
 		{
-			if (LeadingKey && Params.TimeToPixel.FrameDeltaToPixel(PrecomputedKeys[0].FinalKeyPosition - LeadingKey->Time) > Params.KeySizePx.X)
+			if (LeadingKey && LeadingKeyParams->ConnectionStyle != EKeyConnectionStyle::None &&
+				Params.TimeToPixel.FrameDeltaToPixel(PrecomputedKeys[0].FinalKeyPosition - LeadingKey->Time) > Params.KeySizePx.X)
 			{
 				PrecomputedKeyBars.Add(FKeyBar{ LeadingKey->Time, PrecomputedKeys[0].FinalKeyPosition, EKeyBarRenderingFlags::None, LeadingKeyParams->ConnectionStyle });
 			}
@@ -527,13 +528,15 @@ void FKeyRenderer::UpdateViewDependentData(const FKeyBatchParameters& Params, co
 			{
 				FKey ThisKey = PrecomputedKeys[Index];
 				FKey NextKey = PrecomputedKeys[Index+1];
-				if (Params.TimeToPixel.FrameDeltaToPixel(NextKey.FinalKeyPosition - ThisKey.FinalKeyPosition) > Params.KeySizePx.X)
+				if (ThisKey.Params.ConnectionStyle != EKeyConnectionStyle::None &&
+					Params.TimeToPixel.FrameDeltaToPixel(NextKey.FinalKeyPosition - ThisKey.FinalKeyPosition) > Params.KeySizePx.X)
 				{
 					PrecomputedKeyBars.Add(FKeyBar{ ThisKey.FinalKeyPosition, NextKey.FinalKeyPosition, EKeyBarRenderingFlags::None, ThisKey.Params.ConnectionStyle });
 				}
 			}
 
-			if (TrailingKey && Params.TimeToPixel.FrameDeltaToPixel(TrailingKey->Time - PrecomputedKeys.Last().FinalKeyPosition) > Params.KeySizePx.X)
+			if (TrailingKey && PrecomputedKeys.Last().Params.ConnectionStyle != EKeyConnectionStyle::None &&
+				Params.TimeToPixel.FrameDeltaToPixel(TrailingKey->Time - PrecomputedKeys.Last().FinalKeyPosition) > Params.KeySizePx.X)
 			{
 				PrecomputedKeyBars.Add(FKeyBar{ PrecomputedKeys.Last().FinalKeyPosition, TrailingKey->Time, EKeyBarRenderingFlags::None, PrecomputedKeys.Last().Params.ConnectionStyle });
 			}
