Sequencer: Fixed another erroneous fwd declaration

diff --git a/Engine/Source/Editor/SequencerCore/Public/MVVM/Views/SCompoundTrackLaneView.h b/Engine/Source/Editor/SequencerCore/Public/MVVM/Views/SCompoundTrackLaneView.h
--- a/Engine/Source/Editor/SequencerCore/Public/MVVM/Views/SCompoundTrackLaneView.h
+++ b/Engine/Source/Editor/SequencerCore/Public/MVVM/Views/SCompoundTrackLaneView.h
@@ -13,7 +13,6 @@ struct FGeometry;
 
 class FPaintArgs;
 class FSlateWindowElementList;
-class FTrackAreaViewModel;
 
 namespace UE::Sequencer
 {
