Sequencer: Add channel view to collapsed sections

SSequencerSection now has an SChannelView that is only shown when collapsed, which reflects all the keys within it (including top-level channels). This actually allows us to also remove a second pass we had to do previously to generate top-level channel view widgets as well.

diff --git a/Engine/Source/Editor/Sequencer/Private/MVVM/ViewModels/TrackModelLayoutBuilder.cpp b/Engine/Source/Editor/Sequencer/Private/MVVM/ViewModels/TrackModelLayoutBuilder.cpp
--- a/Engine/Source/Editor/Sequencer/Private/MVVM/ViewModels/TrackModelLayoutBuilder.cpp
+++ b/Engine/Source/Editor/Sequencer/Private/MVVM/ViewModels/TrackModelLayoutBuilder.cpp
@@ -136,6 +136,8 @@ FTrackModelLayoutBuilder::~FTrackModelLayoutBuilder()
 
 void FTrackModelLayoutBuilder::RefreshLayout(TSharedPtr<FSectionModel> InSection)
 {
+	InSection->SetLinkedOutlinerItem(Root);
+
 	// Reset everything
 	OutlinerList.Reset();
 
diff --git a/Engine/Source/Editor/Sequencer/Private/MVVM/Views/SChannelView.cpp b/Engine/Source/Editor/Sequencer/Private/MVVM/Views/SChannelView.cpp
--- a/Engine/Source/Editor/Sequencer/Private/MVVM/Views/SChannelView.cpp
+++ b/Engine/Source/Editor/Sequencer/Private/MVVM/Views/SChannelView.cpp
@@ -415,8 +415,10 @@ int32 SChannelView::OnPaint(const FPaintArgs& Args, const FGeometry& AllottedGeo
 		return LayerId;
 	}
 
+	const bool bIncludeThis = true;
+
 	TViewModelPtr<FTrackAreaViewModel> TrackArea = WeakTrackArea.Pin();
-	TViewModelPtr<FSectionModel>  Section  = Model->FindAncestorOfType<FSectionModel>();
+	TViewModelPtr<FSectionModel>  Section  = Model->FindAncestorOfType<FSectionModel>(bIncludeThis);
 	TViewModelPtr<FSequenceModel> Sequence = Model->FindAncestorOfType<FSequenceModel>();
 
 	UMovieSceneSection* SectionObject = Section->GetSection();
@@ -596,7 +598,8 @@ TViewModelPtr<FSectionModel> SChannelView::GetSection() const
 	FViewModelPtr Model = WeakModel.Pin();
 	if (Model)
 	{
-		return Model->FindAncestorOfType<FSectionModel>();
+		constexpr bool bIncludeThis = true;
+		return Model->FindAncestorOfType<FSectionModel>(bIncludeThis);
 	}
 	return nullptr;
 }
diff --git a/Engine/Source/Editor/Sequencer/Private/SSequencerSection.cpp b/Engine/Source/Editor/Sequencer/Private/SSequencerSection.cpp
--- a/Engine/Source/Editor/Sequencer/Private/SSequencerSection.cpp
+++ b/Engine/Source/Editor/Sequencer/Private/SSequencerSection.cpp
@@ -10,6 +10,7 @@
 #include "MVVM/Views/STrackAreaView.h"
 #include "MVVM/Views/SCompoundTrackLaneView.h"
 #include "MVVM/Views/STrackLane.h"
+#include "MVVM/Views/SChannelView.h"
 #include "MVVM/Extensions/IObjectBindingExtension.h"
 #include "Rendering/DrawElements.h"
 #include "Styling/AppStyle.h"
@@ -845,6 +846,12 @@ void SSequencerSection::Construct( const FArguments& InArgs, TSharedPtr<FSequenc
 		[
 			SectionInterface->GenerateSectionWidget()
 		]
+
+		+ SOverlay::Slot()
+		[
+			SNew(SChannelView, InSectionModel, OwningTrackLane->GetTimeToPixel(), OwningTrackLane->GetTrackAreaView()->GetViewModel())
+			.Visibility(this, &SSequencerSection::GetTopLevelChannelGroupVisibility)
+		]
 	];
 }
 
@@ -852,6 +859,20 @@ SSequencerSection::~SSequencerSection()
 {
 }
 
+EVisibility SSequencerSection::GetTopLevelChannelGroupVisibility() const
+{
+	using namespace UE::MovieScene;
+
+	TSharedPtr<FSectionModel> SectionModel = WeakSectionModel.Pin();
+	TViewModelPtr<IOutlinerExtension> Outliner = SectionModel ? SectionModel->FindAncestorOfType<IOutlinerExtension>() : nullptr;
+
+	if (!Outliner || Outliner->IsExpanded())
+	{
+		return EVisibility::Collapsed;
+	}
+	return EVisibility::Visible;
+}
+
 FText SSequencerSection::GetToolTipText() const
 {
 	const UMovieSceneSection* SectionObject = SectionInterface->GetSectionObject();
diff --git a/Engine/Source/Editor/Sequencer/Private/SSequencerSection.h b/Engine/Source/Editor/Sequencer/Private/SSequencerSection.h
--- a/Engine/Source/Editor/Sequencer/Private/SSequencerSection.h
+++ b/Engine/Source/Editor/Sequencer/Private/SSequencerSection.h
@@ -107,6 +107,8 @@ private:
 
 	void OnModifiedIndirectly(UMovieSceneSignedObject* Object) override;
 
+	EVisibility GetTopLevelChannelGroupVisibility() const;
+
 public:
 
 	/** Indicate that the current section selection should throb the specified number of times. A single throb takes 0.2s. */
diff --git a/Engine/Source/Editor/Sequencer/Public/MVVM/ViewModels/SectionModel.h b/Engine/Source/Editor/Sequencer/Public/MVVM/ViewModels/SectionModel.h
--- a/Engine/Source/Editor/Sequencer/Public/MVVM/ViewModels/SectionModel.h
+++ b/Engine/Source/Editor/Sequencer/Public/MVVM/ViewModels/SectionModel.h
@@ -11,6 +11,7 @@
 #include "MVVM/Extensions/ISnappableExtension.h"
 #include "MVVM/Extensions/IDraggableTrackAreaExtension.h"
 #include "MVVM/Extensions/IStretchableExtension.h"
+#include "MVVM/Extensions/LinkedOutlinerExtension.h"
 #include "Delegates/DelegateCombinations.h"
 #include "EventHandlers/ISignedObjectEventHandler.h"
 #include "EventHandlers/ISectionEventHandler.h"
@@ -31,6 +32,7 @@ struct FOverlappingSections;
  */
 class FSectionModel
 	: public FViewModel
+	, public FLinkedOutlinerExtension
 	, public ILayerBarExtension
 	, public ITrackLaneExtension
 	, public ISelectableExtension
@@ -43,6 +45,7 @@ class FSectionModel
 public:
 
 	UE_SEQUENCER_DECLARE_CASTABLE(FSectionModel, FViewModel
+		, FLinkedOutlinerExtension
 		, ILayerBarExtension
 		, ITrackLaneExtension
 		, ISelectableExtension
diff --git a/Engine/Source/Editor/SequencerCore/Private/MVVM/Views/STrackLane.cpp b/Engine/Source/Editor/SequencerCore/Private/MVVM/Views/STrackLane.cpp
--- a/Engine/Source/Editor/SequencerCore/Private/MVVM/Views/STrackLane.cpp
+++ b/Engine/Source/Editor/SequencerCore/Private/MVVM/Views/STrackLane.cpp
@@ -81,6 +81,16 @@ TSharedPtr<STrackLane> STrackLane::GetParentLane() const
 	return ParentLane;
 }
 
+TSharedPtr<STrackAreaView> STrackLane::GetTrackAreaView() const
+{
+	return WeakTrackAreaView.Pin();
+}
+
+TSharedPtr<FTimeToPixel> STrackLane::GetTimeToPixel() const
+{
+	return TimeToPixel;
+}
+
 TArrayView<const TWeakPtr<STrackLane>> STrackLane::GetChildLanes() const
 {
 	return WeakChildLanes;
@@ -167,43 +177,6 @@ void STrackLane::RecreateWidgets()
 		}
 	}
 
-	// Construct views for top level children track lane
-	for (TTypedIterator<ITrackAreaExtension, FViewModelVariantIterator> TopLevelIt(TrackAreaExtension->GetTopLevelChildTrackAreaModels()); TopLevelIt; ++TopLevelIt)
-	{
-		TViewModelPtr<ITrackAreaExtension> InlineModel = *TopLevelIt;
-		if (InlineModel->GetTrackAreaParameters().LaneType != ETrackAreaLaneType::None)
-		{
-			// Add these widgets directly to this lane as well
-			for (TTypedIterator<ITrackLaneExtension, FViewModelVariantIterator> It(InlineModel->GetTrackAreaModelList()); It; ++It)
-			{
-				TViewModelPtr<ITrackLaneExtension> Model = *It;
-				TSharedPtr<ITrackLaneWidget> ParentView;
-				FViewModelPtr Parent = Model.AsModel()->GetParent();
-				// Find the parent model that exists in the parent lane
-				while (Parent && !ParentView)
-				{
-					ParentView = FindWidgetForModel(Parent);
-					Parent = Parent->GetParent();
-				}
-
-				TSharedPtr<ITrackLaneWidget> NewView = Model->CreateTrackLaneView(ViewParams);
-				WidgetsByModel.Add(Model, NewView);
-
-				// If we have a parent view, add the widget to that
-				if (ParentView)
-				{
-					ParentView->AddChildView(NewView, SharedThis(this));
-				}
-				// Otherwise we add it directly to this lane
-				else
-				{
-					FSlot::FSlotArguments SlotArguments(MakeUnique<FSlot>(NewView));
-					SlotArguments.AttachWidget(NewView->AsWidget());
-					Children.AddSlot(MoveTemp(SlotArguments));
-				}
-			}
-		}
-	}
 }
 
 void STrackLane::OnArrangeChildren( const FGeometry& AllottedGeometry, FArrangedChildren& ArrangedChildren ) const
diff --git a/Engine/Source/Editor/SequencerCore/Public/MVVM/Views/STrackLane.h b/Engine/Source/Editor/SequencerCore/Public/MVVM/Views/STrackLane.h
--- a/Engine/Source/Editor/SequencerCore/Public/MVVM/Views/STrackLane.h
+++ b/Engine/Source/Editor/SequencerCore/Public/MVVM/Views/STrackLane.h
@@ -55,7 +55,9 @@ public:
 	/** Gets the outliner view-model that created this track lane */
 	TViewModelPtr<IOutlinerExtension> GetOutlinerItem() const;
 
+	TSharedPtr<STrackAreaView> GetTrackAreaView() const;
 	TSharedPtr<STrackLane> GetParentLane() const;
+	TSharedPtr<FTimeToPixel> GetTimeToPixel() const;
 
 	TArrayView<const TWeakPtr<STrackLane>> GetChildLanes() const;
 
