Sequencer: Added flags for outliner sizing structure to indicate whether the sizing is static or dynamic (ie, should update every frame)

diff --git a/Engine/Source/Editor/Sequencer/Private/MVVM/ViewModels/CategoryModel.cpp b/Engine/Source/Editor/Sequencer/Private/MVVM/ViewModels/CategoryModel.cpp
--- a/Engine/Source/Editor/Sequencer/Private/MVVM/ViewModels/CategoryModel.cpp
+++ b/Engine/Source/Editor/Sequencer/Private/MVVM/ViewModels/CategoryModel.cpp
@@ -131,20 +131,20 @@ FOutlinerSizing FCategoryGroupModel::RecomputeSizing()
 		if (TSharedPtr<FCategoryModel> Category = WeakCategory.Pin())
 		{
 			FOutlinerSizing Desired = Category->GetDesiredSizing();
-
-			MaxSizing.Height = FMath::Max(MaxSizing.Height, Desired.Height);
-			MaxSizing.PaddingTop = FMath::Max(MaxSizing.PaddingTop, Desired.PaddingTop);
-			MaxSizing.PaddingBottom = FMath::Max(MaxSizing.PaddingBottom, Desired.PaddingBottom);
+			MaxSizing.Accumulate(Desired);
 		}
 	}
 
-	ComputedSizing = MaxSizing;
-
-	for (const TWeakViewModelPtr<FCategoryModel>& WeakCategory : Categories)
+	if (ComputedSizing != MaxSizing)
 	{
-		if (TSharedPtr<FCategoryModel> Category = WeakCategory.Pin())
+		ComputedSizing = MaxSizing;
+
+		for (const TWeakViewModelPtr<FCategoryModel>& WeakCategory : Categories)
 		{
-			Category->SetComputedSizing(MaxSizing);
+			if (TSharedPtr<FCategoryModel> Category = WeakCategory.Pin())
+			{
+				Category->SetComputedSizing(MaxSizing);
+			}
 		}
 	}
 
@@ -153,6 +153,10 @@ FOutlinerSizing FCategoryGroupModel::RecomputeSizing()
 
 FOutlinerSizing FCategoryGroupModel::GetOutlinerSizing() const
 {
+	if (EnumHasAnyFlags(ComputedSizing.Flags, EOutlinerSizingFlags::DynamicSizing))
+	{
+		const_cast<FCategoryGroupModel*>(this)->RecomputeSizing();
+	}
 	return ComputedSizing;
 }
 
diff --git a/Engine/Source/Editor/Sequencer/Private/MVVM/ViewModels/ChannelModel.cpp b/Engine/Source/Editor/Sequencer/Private/MVVM/ViewModels/ChannelModel.cpp
--- a/Engine/Source/Editor/Sequencer/Private/MVVM/ViewModels/ChannelModel.cpp
+++ b/Engine/Source/Editor/Sequencer/Private/MVVM/ViewModels/ChannelModel.cpp
@@ -689,20 +689,20 @@ FOutlinerSizing FChannelGroupOutlinerModel::RecomputeSizing()
 		if (TSharedPtr<FChannelModel> Channel = WeakChannel.Pin())
 		{
 			FOutlinerSizing Desired = Channel->GetDesiredSizing();
-
-			MaxSizing.Height = FMath::Max(MaxSizing.Height, Desired.Height);
-			MaxSizing.PaddingTop = FMath::Max(MaxSizing.PaddingTop, Desired.PaddingTop);
-			MaxSizing.PaddingBottom = FMath::Max(MaxSizing.PaddingBottom, Desired.PaddingBottom);
+			MaxSizing.Accumulate(Desired);
 		}
 	}
 
-	ComputedSizing = MaxSizing;
-
-	for (TWeakViewModelPtr<FChannelModel> WeakChannel : Channels)
+	if (ComputedSizing != MaxSizing)
 	{
-		if (TSharedPtr<FChannelModel> Channel = WeakChannel.Pin())
+		ComputedSizing = MaxSizing;
+
+		for (TWeakViewModelPtr<FChannelModel> WeakChannel : Channels)
 		{
-			Channel->SetComputedSizing(MaxSizing);
+			if (TSharedPtr<FChannelModel> Channel = WeakChannel.Pin())
+			{
+				Channel->SetComputedSizing(MaxSizing);
+			}
 		}
 	}
 
@@ -711,6 +711,10 @@ FOutlinerSizing FChannelGroupOutlinerModel::RecomputeSizing()
 
 FOutlinerSizing FChannelGroupOutlinerModel::GetOutlinerSizing() const
 {
+	if (EnumHasAnyFlags(ComputedSizing.Flags, EOutlinerSizingFlags::DynamicSizing))
+	{
+		const_cast<FChannelGroupOutlinerModel*>(this)->RecomputeSizing();
+	}
 	return ComputedSizing;
 }
 
diff --git a/Engine/Source/Editor/SequencerCore/Public/MVVM/Extensions/IOutlinerExtension.h b/Engine/Source/Editor/SequencerCore/Public/MVVM/Extensions/IOutlinerExtension.h
--- a/Engine/Source/Editor/SequencerCore/Public/MVVM/Extensions/IOutlinerExtension.h
+++ b/Engine/Source/Editor/SequencerCore/Public/MVVM/Extensions/IOutlinerExtension.h
@@ -38,12 +38,19 @@ class IOutlinerExtension;
 enum class EOutlinerSelectionState
 {
 	None                      = 0,
-	SelectedDirectly          = 1 << 1,
-	HasSelectedKeys           = 1 << 2,
-	HasSelectedTrackAreaItems = 1 << 3,
+	SelectedDirectly          = 1 << 0,
+	HasSelectedKeys           = 1 << 1,
+	HasSelectedTrackAreaItems = 1 << 2,
 };
 ENUM_CLASS_FLAGS(EOutlinerSelectionState);
 
+enum class EOutlinerSizingFlags
+{
+	None                      = 0,
+	DynamicSizing             = 1 << 0,
+};
+ENUM_CLASS_FLAGS(EOutlinerSizingFlags);
+
 struct SEQUENCERCORE_API FOutlinerSizing
 {
 	FOutlinerSizing() = default;
@@ -58,14 +65,37 @@ struct SEQUENCERCORE_API FOutlinerSizing
 		, PaddingBottom(UniformPadding)
 	{}
 
+	friend bool operator==(const FOutlinerSizing& A, const FOutlinerSizing& B)
+	{
+		return 
+			A.Height == B.Height && 
+			A.PaddingTop == B.PaddingTop && 
+			A.PaddingBottom == B.PaddingBottom && 
+			A.Flags == B.Flags;
+	}
+	friend bool operator!=(const FOutlinerSizing& A, const FOutlinerSizing& B)
+	{
+		return !(A == B);
+	}
+
 	float GetTotalHeight() const
 	{
 		return Height + PaddingTop + PaddingBottom;
 	}
 
+	void Accumulate(const FOutlinerSizing& Other)
+	{
+		Height        = FMath::Max(Height, Other.Height);
+		PaddingTop    = FMath::Max(PaddingTop, Other.PaddingTop);
+		PaddingBottom = FMath::Max(PaddingBottom, Other.PaddingBottom);
+		Flags |= Other.Flags;
+	}
+
 	float Height = 0.f;
 	float PaddingTop = 0.f;
 	float PaddingBottom = 0.f;
+	EOutlinerSizingFlags Flags = EOutlinerSizingFlags::None;
+
 };
 
 /** Facade class to allow SOutlinerViewRow to be passed around publicly */
