Sequencer: Fix channel key areas being included when filtered out

diff --git a/Engine/Source/Editor/SequencerCore/Private/MVVM/Views/KeyRenderer.cpp b/Engine/Source/Editor/SequencerCore/Private/MVVM/Views/KeyRenderer.cpp
--- a/Engine/Source/Editor/SequencerCore/Private/MVVM/Views/KeyRenderer.cpp
+++ b/Engine/Source/Editor/SequencerCore/Private/MVVM/Views/KeyRenderer.cpp
@@ -214,6 +214,13 @@ void FKeyRenderer::CacheKeyExtensions(const FKeyBatchParameters& Params) const
 	{
 		FViewModelPtr CurrentViewModel = *ChildIt;
 
+		TViewModelPtr<FLinkedOutlinerExtension> OutlinerExtension = CurrentViewModel.ImplicitCast();
+		if (OutlinerExtension && OutlinerExtension->GetLinkedOutlinerItem() && OutlinerExtension->GetLinkedOutlinerItem()->IsFilteredOut())
+		{
+			ChildIt.IgnoreCurrentChildren();
+			continue;
+		}
+
 		TViewModelPtr<IKeyExtension> KeyRenderer = CurrentViewModel.ImplicitCast();
 		if (KeyRenderer)
 		{
