Sequencer: Fix report of hierarchical changes

diff --git a/Engine/Source/Editor/SequencerCore/Private/MVVM/ViewModels/ViewModelHierarchy.cpp b/Engine/Source/Editor/SequencerCore/Private/MVVM/ViewModels/ViewModelHierarchy.cpp
--- a/Engine/Source/Editor/SequencerCore/Private/MVVM/ViewModels/ViewModelHierarchy.cpp
+++ b/Engine/Source/Editor/SequencerCore/Private/MVVM/ViewModels/ViewModelHierarchy.cpp
@@ -296,16 +296,19 @@ void FViewModelChildren::MoveChildrenTo(const FViewModelChildren& OutDestination
 		// Already same parent so we're just moving from one list to another
 		ListTail = GetTail();
 	}
-	else for (const FViewModelPtr& Child : IterateSubList().ToArray())
+	else
 	{
 		// We purposefully only report changes for the first element
 		// Since that is all that is necessary to trigger events for the old parent
 		bool bReportChanges = true;
 
-		Child->SetParentOnly(OutDestination.Owner, bReportChanges);
-		ListTail = Child;
+		for (const FViewModelPtr& Child : IterateSubList().ToArray())
+		{
+			Child->SetParentOnly(OutDestination.Owner, bReportChanges);
+			ListTail = Child;
 
-		bReportChanges = false;
+			bReportChanges = false;
+		}
 	}
 
 	check(ListTail && ListTail->Link.Next == nullptr);
