Sequencer: Added an ensure to protect against adding a model as a child of itself

diff --git a/Engine/Source/Editor/SequencerCore/Private/MVVM/ViewModels/ViewModelHierarchy.cpp b/Engine/Source/Editor/SequencerCore/Private/MVVM/ViewModels/ViewModelHierarchy.cpp
--- a/Engine/Source/Editor/SequencerCore/Private/MVVM/ViewModels/ViewModelHierarchy.cpp
+++ b/Engine/Source/Editor/SequencerCore/Private/MVVM/ViewModels/ViewModelHierarchy.cpp
@@ -225,6 +225,10 @@ FViewModelSubListIterator FViewModelChildren::IterateSubList() const
 
 void FViewModelChildren::AddChild(const TSharedPtr<FViewModel>& Child)
 {
+	if (!ensureMsgf(Child != Owner, TEXT("Cannot attach a model to itself")))
+	{
+		return;
+	}
 	if (!ensureMsgf(Child, TEXT("Unable to add a null child")))
 	{
 		return;
