Sequencer: fix loophole where non-renameable tracks could be renamed by simply clicking on their label in the sequencer outliner view

diff --git a/Engine/Source/Editor/SequencerCore/Private/MVVM/Views/SOutlinerItemViewBase.cpp b/Engine/Source/Editor/SequencerCore/Private/MVVM/Views/SOutlinerItemViewBase.cpp
--- a/Engine/Source/Editor/SequencerCore/Private/MVVM/Views/SOutlinerItemViewBase.cpp
+++ b/Engine/Source/Editor/SequencerCore/Private/MVVM/Views/SOutlinerItemViewBase.cpp
@@ -99,7 +99,7 @@ void SOutlinerItemViewBase::Construct(
 	if (TViewModelPtr<IRenameableExtension> Renameable = WeakRenameExtension.Pin())
 	{
 		TSharedRef<SInlineEditableTextBlock> EditableLabel = SNew(SInlineEditableTextBlock)
-		.IsReadOnly(IsReadOnlyAttribute)
+		.IsReadOnly(this, &SOutlinerItemViewBase::IsNodeLabelReadOnly)
 		.Font(this, &SOutlinerItemViewBase::GetLabelFont)
 		.Text(this, &SOutlinerItemViewBase::GetLabel)
 		.ToolTipText(this, &SOutlinerItemViewBase::GetLabelToolTipText)
@@ -318,10 +318,11 @@ FText SOutlinerItemViewBase::GetIconToolTipText() const
 
 bool SOutlinerItemViewBase::IsNodeLabelReadOnly() const
 {
+	const bool bIsReadOnly = IsReadOnlyAttribute.Get();
 	TSharedPtr<FEditorViewModel> Editor = WeakEditor.Pin();
 	TViewModelPtr<IRenameableExtension> Renamable = WeakRenameExtension.Pin();
 
-	return !Editor || Editor->IsReadOnly() || !Renamable || !Renamable->CanRename();
+	return bIsReadOnly || !Editor || Editor->IsReadOnly() || !Renamable || !Renamable->CanRename();
 }
 
 bool SOutlinerItemViewBase::IsDimmed() const
