Sequencer: Ensure that shared data is correctly propagated to children when detaching models from a hierarchy

diff --git a/Engine/Source/Editor/SequencerCore/Private/MVVM/ViewModels/ViewModel.cpp b/Engine/Source/Editor/SequencerCore/Private/MVVM/ViewModels/ViewModel.cpp
--- a/Engine/Source/Editor/SequencerCore/Private/MVVM/ViewModels/ViewModel.cpp
+++ b/Engine/Source/Editor/SequencerCore/Private/MVVM/ViewModels/ViewModel.cpp
@@ -107,7 +107,7 @@ FViewModel::ESetParentResult FViewModel::SetParentOnly(const TSharedPtr<FViewMod
 	else
 	{
 		WeakParent = nullptr;
-		SharedData = nullptr;
+		SetSharedData(nullptr);
 
 		return ESetParentResult::ClearedParent;
 	}
