Sequencer: Fixed shared data not necessarily being initialized in OnConstruct calls for child models

When a model was being (re)constructed after a change (in the case of the bug it was moving a sub-folder out of its parent) it was possible for the shared data of its children to be null, whilst its parent was non-null. This creates the issue that the parent recycling its children would not trigger hierarchy change notifications for those children (because there is no shared data to trigger the event on). For layerbars this means leaving a dangling ptr in FViewModelExtensionCollection.

Really this breaks the invariants of the shared data - it should not ever be possible for a child to have different shared data than its parent (and that invariant was reliable outside of these OnConstruct calls). To correct it we now assign all shared data before calling OnConstruct/OnDestruct. In order to prevent multiple OnConstruct calls we now have to track whether it models need construction.

Also fixed destruction of FScopedViewModelListHead not calling hierarchy changed events by explicity emptying the list on destruction. This technically is not required to fix the bug, but I think is good practice nonetheless to avoid future problems with models being destructed while they are still parented.

diff --git a/Engine/Source/Editor/SequencerCore/Private/MVVM/ViewModels/ViewModel.cpp b/Engine/Source/Editor/SequencerCore/Private/MVVM/ViewModels/ViewModel.cpp
--- a/Engine/Source/Editor/SequencerCore/Private/MVVM/ViewModels/ViewModel.cpp
+++ b/Engine/Source/Editor/SequencerCore/Private/MVVM/ViewModels/ViewModel.cpp
@@ -17,6 +17,7 @@ namespace Sequencer
 FViewModel::FViewModel()
 	: FirstChildListHead(nullptr)
 	, ActiveIterationCount(0)
+	, bNeedsConstruction(1)
 {
 }
 
@@ -27,7 +28,7 @@ FViewModel::~FViewModel()
 
 bool FViewModel::IsConstructed() const
 {
-	return SharedData != nullptr;
+	return SharedData != nullptr && bNeedsConstruction == false;
 }
 
 void FViewModel::RegisterChildList(FViewModelListHead* InChildren)
@@ -159,34 +160,49 @@ void FViewModel::SetSharedData(TSharedPtr<FSharedViewModelData> InSharedData)
 		return;
 	}
 
-	SharedData = InSharedData;
-
 	if (InSharedData)
 	{
-		OnConstruct();
-	}
-	else
-	{
-		OnDestruct();
-	}
+		bNeedsConstruction = true;
+		SharedData = InSharedData;
 
-	for (const FViewModelPtr& Child : GetDescendants())
-	{
-		// This is safe to call because we're inside a parent-first iterator.
-		// Any children that are added inside here will get implicitly
-		// included by the iterator, but they may have already been constructed
-		if (Child->SharedData != InSharedData)
+		for (const FViewModelPtr& Child : GetDescendants())
 		{
+			Child->bNeedsConstruction = true;
 			Child->SharedData = InSharedData;
+		}
+
+		OnConstruct();
+		bNeedsConstruction = false;
 
-			if (InSharedData)
+		for (const FViewModelPtr& Child : GetDescendants())
+		{
+			if (Child->bNeedsConstruction)
 			{
+				// This is safe to call because we're inside a parent-first iterator.
+				// Any children that are added inside here will get implicitly
+				// included by the iterator, but they may have already been constructed
 				Child->OnConstruct();
+				Child->bNeedsConstruction = false;
 			}
-			else
-			{
-				Child->OnDestruct();
-			}
+		}
+	}
+	else
+	{
+		SharedData = nullptr;
+
+		for (const FViewModelPtr& Child : GetDescendants())
+		{
+			Child->SharedData = nullptr;
+		}
+
+		OnDestruct();
+
+		for (const FViewModelPtr& Child : GetDescendants())
+		{
+			// This is safe to call because we're inside a parent-first iterator.
+			// Any children that are removed inside here will get implicitly
+			// excluded by the iterator
+			Child->OnDestruct();
 		}
 	}
 }
diff --git a/Engine/Source/Editor/SequencerCore/Private/MVVM/ViewModels/ViewModelHierarchy.cpp b/Engine/Source/Editor/SequencerCore/Private/MVVM/ViewModels/ViewModelHierarchy.cpp
--- a/Engine/Source/Editor/SequencerCore/Private/MVVM/ViewModels/ViewModelHierarchy.cpp
+++ b/Engine/Source/Editor/SequencerCore/Private/MVVM/ViewModels/ViewModelHierarchy.cpp
@@ -174,6 +174,9 @@ FScopedViewModelListHead::FScopedViewModelListHead(const TSharedPtr<FViewModel>&
 
 FScopedViewModelListHead::~FScopedViewModelListHead()
 {
+	// Destroy everything in this list
+	GetChildren().Empty();
+
 	FViewModelListHead* ThisList = &ListHead;
 	FViewModelListHead** Ptr = &Model->FirstChildListHead;
 	while (*Ptr)
diff --git a/Engine/Source/Editor/SequencerCore/Public/MVVM/ViewModels/ViewModel.h b/Engine/Source/Editor/SequencerCore/Public/MVVM/ViewModels/ViewModel.h
--- a/Engine/Source/Editor/SequencerCore/Public/MVVM/ViewModels/ViewModel.h
+++ b/Engine/Source/Editor/SequencerCore/Public/MVVM/ViewModels/ViewModel.h
@@ -286,6 +286,11 @@ private:
 	         looking at). Such models are safe to move due to children being linked lists
 	 */
 	int32 ActiveIterationCount;
+
+	/**
+	 * Flag to track whether this view model needs to be constructed
+	 */
+	uint8 bNeedsConstruction : 1;
 };
 
 } // namespace Sequencer
