Sequencer: Tightened up implicit viewmodel conversions, introduced static type information for ICastable instances

In order to support incoming improvements to the selection API, it is required that we no longer allow implicit construction of view model ptr types.
FViewModelPtr itself has now become a template alias of TViewModelPtr rather than a base class.
Casting between model types is now implemented through a static type table structure for each instance that stores type IDs and class offsets for each type, and no longer requires any virtual interface on the base class.
Casting models measured ~30% faster using this approach in my tests.
This will later be used by the scripting API to provide type information as well.
A new sparse bit set type has been added that provides the foundation for storing type information with the minimal amount of memory.

diff --git a/Engine/Source/Editor/Sequencer/Private/MVVM/CurveEditorIntegrationExtension.cpp b/Engine/Source/Editor/Sequencer/Private/MVVM/CurveEditorIntegrationExtension.cpp
--- a/Engine/Source/Editor/Sequencer/Private/MVVM/CurveEditorIntegrationExtension.cpp
+++ b/Engine/Source/Editor/Sequencer/Private/MVVM/CurveEditorIntegrationExtension.cpp
@@ -77,7 +77,7 @@ void FCurveEditorIntegrationExtension::UpdateCurveEditor()
 	{
 		bool bIsRelevant = false;
 		const FCurveEditorTreeItemID ItemID(It.Value());
-		TViewModelPtr<ICurveEditorTreeItemExtension> ViewModel(It.Key().Pin());
+		TViewModelPtr<ICurveEditorTreeItemExtension> ViewModel = CastViewModel<ICurveEditorTreeItemExtension>(It.Key().Pin());
 		if (ViewModel)
 		{
 			TViewModelPtr<IOutlinerExtension> OutlinerModel = ViewModel.ImplicitCast();
@@ -129,7 +129,7 @@ void FCurveEditorIntegrationExtension::UpdateCurveEditor()
 			const FCurveEditorTreeItemID ItemID = ViewModelToTreeItemIDMap.FindRef(ChildViewModel.AsModel());
 			if (!ItemID.IsValid())
 			{
-				AddToCurveEditor(ChildViewModel.AsModel(), CurveEditor);
+				AddToCurveEditor(ChildViewModel, CurveEditor);
 			}
 		}
 	}
@@ -156,9 +156,9 @@ FCurveEditorTreeItemID FCurveEditorIntegrationExtension::AddToCurveEditor(TViewM
 	}
 
 	// Recursively create any needed parent curve editor items
-	TSharedPtr<FViewModel> Parent = InViewModel.AsModel()->GetParent();
+	TViewModelPtr<ICurveEditorTreeItemExtension> Parent = InViewModel.AsModel()->CastParent<ICurveEditorTreeItemExtension>();
 
-	FCurveEditorTreeItemID ParentID = Parent.IsValid() ?
+	FCurveEditorTreeItemID ParentID = Parent ?
 		AddToCurveEditor(Parent, InCurveEditor) :
 		FCurveEditorTreeItemID::Invalid();
 
@@ -188,10 +188,10 @@ void FCurveEditorIntegrationExtension::ResetCurveEditor()
 		return;
 	}
 
-	for (auto Pair : ViewModelToTreeItemIDMap)
+	for (TPair<TWeakPtr<FViewModel>, FCurveEditorTreeItemID> Pair : ViewModelToTreeItemIDMap)
 	{
 		FCurveEditorTreeItemID ItemID(Pair.Value);
-		TViewModelPtr<ICurveEditorTreeItemExtension> ViewModel(Pair.Key.Pin());
+		TViewModelPtr<ICurveEditorTreeItemExtension> ViewModel = CastViewModel<ICurveEditorTreeItemExtension>(Pair.Key.Pin());
 		if (ViewModel)
 		{
 			ViewModel->OnRemovedFromCurveEditor(CurveEditor);
diff --git a/Engine/Source/Editor/Sequencer/Private/MVVM/ViewModels/CategoryModel.cpp b/Engine/Source/Editor/Sequencer/Private/MVVM/ViewModels/CategoryModel.cpp
--- a/Engine/Source/Editor/Sequencer/Private/MVVM/ViewModels/CategoryModel.cpp
+++ b/Engine/Source/Editor/Sequencer/Private/MVVM/ViewModels/CategoryModel.cpp
@@ -215,7 +215,7 @@ bool FCategoryGroupModel::CanDelete(FText* OutErrorMessage) const
 void FCategoryGroupModel::Delete()
 {
 	TArray<FName> PathFromTrack;
-	TViewModelPtr<ITrackExtension> Track = GetParentTrackNodeAndNamePath(this, PathFromTrack);
+	TViewModelPtr<ITrackExtension> Track = GetParentTrackNodeAndNamePath(CastThisShared<IOutlinerExtension>(), PathFromTrack);
 	check(Track);
 
 	Track->GetTrack()->Modify();
diff --git a/Engine/Source/Editor/Sequencer/Private/MVVM/ViewModels/ChannelModel.cpp b/Engine/Source/Editor/Sequencer/Private/MVVM/ViewModels/ChannelModel.cpp
--- a/Engine/Source/Editor/Sequencer/Private/MVVM/ViewModels/ChannelModel.cpp
+++ b/Engine/Source/Editor/Sequencer/Private/MVVM/ViewModels/ChannelModel.cpp
@@ -309,7 +309,7 @@ void FChannelGroupModel::OnChannelOverridesChanged()
 
 void FChannelGroupModel::CleanupChannels()
 {
-	const int32 NumRemoved = Channels.RemoveAll([](TWeakViewModelPtr<FChannelModel> Item) { return !Item.IsValid(); });
+	const int32 NumRemoved = Channels.RemoveAll([](TWeakViewModelPtr<FChannelModel> Item) { return !Item.Pin().IsValid(); });
 	if (NumRemoved > 0)
 	{
 		++ChannelsSerialNumber;
@@ -827,7 +827,7 @@ bool FChannelGroupOutlinerModel::CanDelete(FText* OutErrorMessage) const
 void FChannelGroupOutlinerModel::Delete()
 {
 	TArray<FName> PathFromTrack;
-	TViewModelPtr<ITrackExtension> Track = GetParentTrackNodeAndNamePath(this, PathFromTrack);
+	TViewModelPtr<ITrackExtension> Track = GetParentTrackNodeAndNamePath(SharedThis(this), PathFromTrack);
 
 	Track->GetTrack()->Modify();
 
diff --git a/Engine/Source/Editor/Sequencer/Private/MVVM/ViewModels/FolderModel.cpp b/Engine/Source/Editor/Sequencer/Private/MVVM/ViewModels/FolderModel.cpp
--- a/Engine/Source/Editor/Sequencer/Private/MVVM/ViewModels/FolderModel.cpp
+++ b/Engine/Source/Editor/Sequencer/Private/MVVM/ViewModels/FolderModel.cpp
@@ -105,7 +105,7 @@ void FFolderModel::RepopulateChildren()
 		if (Sequencer->GetSequencerSettings()->GetShowLayerBars())
 		{
 			LayerBar = MakeShared<FLayerBarModel>(This);
-			LayerBar->SetLinkedOutlinerItem(AsShared());
+			LayerBar->SetLinkedOutlinerItem(SharedThis(this));
 
 			TrackAreaChildren.AddChild(LayerBar);
 		}
diff --git a/Engine/Source/Editor/Sequencer/Private/MVVM/ViewModels/ObjectBindingModel.cpp b/Engine/Source/Editor/Sequencer/Private/MVVM/ViewModels/ObjectBindingModel.cpp
--- a/Engine/Source/Editor/Sequencer/Private/MVVM/ViewModels/ObjectBindingModel.cpp
+++ b/Engine/Source/Editor/Sequencer/Private/MVVM/ViewModels/ObjectBindingModel.cpp
@@ -203,7 +203,7 @@ void FObjectBindingModel::OnConstruct()
 		if (Sequencer->GetSequencerSettings()->GetShowLayerBars())
 		{
 			LayerBar = MakeShared<FLayerBarModel>(AsShared());
-			LayerBar->SetLinkedOutlinerItem(AsShared());
+			LayerBar->SetLinkedOutlinerItem(SharedThis(this));
 
 			GetChildrenForList(&TrackAreaList).AddChild(LayerBar);
 		}
diff --git a/Engine/Source/Editor/Sequencer/Private/MVVM/ViewModels/OutlinerItemModel.cpp b/Engine/Source/Editor/Sequencer/Private/MVVM/ViewModels/OutlinerItemModel.cpp
--- a/Engine/Source/Editor/Sequencer/Private/MVVM/ViewModels/OutlinerItemModel.cpp
+++ b/Engine/Source/Editor/Sequencer/Private/MVVM/ViewModels/OutlinerItemModel.cpp
@@ -260,7 +260,7 @@ bool FOutlinerItemModelMixin::IsDimmed() const
 	TSharedPtr<FSequencer> Sequencer = EditorViewModel ? EditorViewModel->GetSequencerImpl() : nullptr;
 	if (Sequencer)
 	{
-		TSharedPtr<FViewModel> ThisShared(const_cast<FViewModel*>(AsViewModel())->AsShared());
+		TViewModelPtr<IOutlinerExtension> ThisShared(const_cast<FViewModel*>(AsViewModel())->AsShared(), const_cast<FOutlinerItemModelMixin*>(this));
 		FSequencerNodeTree& ParentTree = Sequencer->GetNodeTree().Get();
 		if (ParentTree.IsNodeMute(ThisShared) || (ParentTree.HasSoloNodes() && !ParentTree.IsNodeSolo(ThisShared)))
 		{
@@ -276,7 +276,7 @@ bool FOutlinerItemModelMixin::IsSolo() const
 	TSharedPtr<FSequencer> Sequencer = EditorViewModel ? EditorViewModel->GetSequencerImpl() : nullptr;
 	if (Sequencer)
 	{
-		TSharedPtr<FViewModel> ThisShared(const_cast<FViewModel*>(AsViewModel())->AsShared());
+		TViewModelPtr<IOutlinerExtension> ThisShared(const_cast<FViewModel*>(AsViewModel())->AsShared(), const_cast<FOutlinerItemModelMixin*>(this));
 		return Sequencer->GetNodeTree()->IsNodeSolo(ThisShared);
 	}
 	return false;
@@ -288,7 +288,7 @@ bool FOutlinerItemModelMixin::IsMuted() const
 	TSharedPtr<FSequencer> Sequencer = EditorViewModel ? EditorViewModel->GetSequencerImpl() : nullptr;
 	if (Sequencer)
 	{
-		TSharedPtr<FViewModel> ThisShared(const_cast<FViewModel*>(AsViewModel())->AsShared());
+		TViewModelPtr<IOutlinerExtension> ThisShared(const_cast<FViewModel*>(AsViewModel())->AsShared(), const_cast<FOutlinerItemModelMixin*>(this));
 		return Sequencer->GetNodeTree()->IsNodeMute(ThisShared);
 	}
 	return false;
diff --git a/Engine/Source/Editor/Sequencer/Private/MVVM/ViewModels/TrackModelLayoutBuilder.cpp b/Engine/Source/Editor/Sequencer/Private/MVVM/ViewModels/TrackModelLayoutBuilder.cpp
--- a/Engine/Source/Editor/Sequencer/Private/MVVM/ViewModels/TrackModelLayoutBuilder.cpp
+++ b/Engine/Source/Editor/Sequencer/Private/MVVM/ViewModels/TrackModelLayoutBuilder.cpp
@@ -144,7 +144,7 @@ FTrackModelLayoutBuilder::~FTrackModelLayoutBuilder()
 
 void FTrackModelLayoutBuilder::RefreshLayout(TSharedPtr<FSectionModel> InSection)
 {
-	InSection->SetLinkedOutlinerItem(Root);
+	InSection->SetLinkedOutlinerItem(CastViewModelChecked<IOutlinerExtension>(Root));
 
 	// Reset everything
 	OutlinerList.Reset();
@@ -312,11 +312,11 @@ void FTrackModelLayoutBuilder::AddChannel(const FMovieSceneChannelHandle& Channe
 		OutlinerModel->AddChannel(TrackAreaModel);
 		if (bIsTopLevel)
 		{
-			TrackAreaModel->SetLinkedOutlinerItem(OutlinerModel->GetParent());
+			TrackAreaModel->SetLinkedOutlinerItem(OutlinerModel->GetParent().ImplicitCastChecked());
 		}
 		else
 		{
-			TrackAreaModel->SetLinkedOutlinerItem(OutlinerModel.ImplicitCast());
+			TrackAreaModel->SetLinkedOutlinerItem(OutlinerModel.ImplicitCastChecked());
 		}
 		TrackAreaModel->Initialize(SequencerSection, Channel);
 
diff --git a/Engine/Source/Editor/Sequencer/Private/MVVM/Views/SChannelView.cpp b/Engine/Source/Editor/Sequencer/Private/MVVM/Views/SChannelView.cpp
--- a/Engine/Source/Editor/Sequencer/Private/MVVM/Views/SChannelView.cpp
+++ b/Engine/Source/Editor/Sequencer/Private/MVVM/Views/SChannelView.cpp
@@ -557,7 +557,7 @@ int32 SChannelView::DrawLane(const FPaintArgs& Args, const FGeometry& AllottedGe
 
 	FLinearColor HighlightColor;
 	bool bDrawHighlight = false;
-	if (Sequencer->GetSelection().NodeHasSelectedKeysOrSections(OutlinerItem.AsModel()))
+	if (Sequencer->GetSelection().NodeHasSelectedKeysOrSections(OutlinerItem))
 	{
 		bDrawHighlight = true;
 		HighlightColor = FLinearColor(1.0f, 1.0f, 1.0f, 0.15f);
diff --git a/Engine/Source/Editor/Sequencer/Private/SSequencer.cpp b/Engine/Source/Editor/Sequencer/Private/SSequencer.cpp
--- a/Engine/Source/Editor/Sequencer/Private/SSequencer.cpp
+++ b/Engine/Source/Editor/Sequencer/Private/SSequencer.cpp
@@ -3468,7 +3468,7 @@ void SSequencer::StepToKey(bool bStepToNextKey, bool bCameraOnly)
 			Nodes.Reserve(Sequencer->GetSelection().GetSelectedOutlinerItems().Num());
 			for (const TWeakPtr<FViewModel>& WeakModel : Sequencer->GetSelection().GetSelectedOutlinerItems())
 			{
-				TViewModelPtr<IOutlinerExtension> Model = WeakModel.Pin();
+				TViewModelPtr<IOutlinerExtension> Model = CastViewModel<IOutlinerExtension>(WeakModel.Pin());
 				if (Model)
 				{
 					Nodes.Add(Model);
@@ -3831,7 +3831,7 @@ FPasteContextMenuArgs SSequencer::GeneratePasteArgs(FFrameNumber PasteAtTime, TS
 		{
 			if (SelectedNode->IsA<FCategoryGroupModel>() || SelectedNode->IsA<ITrackExtension>() || SelectedNode->IsA<FChannelGroupModel>())
 			{
-				TViewModelPtr<IOutlinerExtension> TrackModel(SelectedNode);
+				TViewModelPtr<IOutlinerExtension> TrackModel =  CastViewModel<IOutlinerExtension>(SelectedNode);
 				if (TrackModel)
 				{
 					PasteIntoNodes.Add(TrackModel);
diff --git a/Engine/Source/Editor/Sequencer/Private/Sequencer.cpp b/Engine/Source/Editor/Sequencer/Private/Sequencer.cpp
--- a/Engine/Source/Editor/Sequencer/Private/Sequencer.cpp
+++ b/Engine/Source/Editor/Sequencer/Private/Sequencer.cpp
@@ -6281,7 +6281,7 @@ void FSequencer::SynchronizeExternalSelectionWithSequencerSelection()
 
 	for ( TWeakPtr<FViewModel> WeakItem : Items)
 	{
-		TSharedPtr<FViewModel> SelectedItem = WeakItem.Pin();
+		TViewModelPtr<FViewModel> SelectedItem = WeakItem.Pin();
 		if (!SelectedItem)
 		{
 			continue;
@@ -6290,8 +6290,8 @@ void FSequencer::SynchronizeExternalSelectionWithSequencerSelection()
 		//HACK for DHI, if we have an active control rig then one is selected so don't find a parent actor or compomonent to select	
 		//but if we do select the actor/compoent directly we still select it.
 		TViewModelPtr<IObjectBindingExtension> ObjectBinding = bHACK_ControlRigEditMode
-			? FViewModelPtr(SelectedItem)
-			: SelectedItem->FindAncestorOfType(IObjectBindingExtension::ID, true);
+			? CastViewModel<IObjectBindingExtension>(SelectedItem)
+			: SelectedItem->FindAncestorOfType<IObjectBindingExtension>(true);
 
 		// If the closest node is an object node, try to get the actor/component nodes from it.
 		if (ObjectBinding)
@@ -6560,13 +6560,13 @@ void FSequencer::SynchronizeSequencerSelectionWithExternalSelection()
 		{
 			for (TWeakPtr<FViewModel> WeakNode : OutlinerSelection)
 			{
-				if (TSharedPtr<FViewModel> Node = WeakNode.Pin())
+				if (TViewModelPtr<IOutlinerExtension> Node = CastViewModel<IOutlinerExtension>(WeakNode.Pin()))
 				{
-					TSharedPtr<FViewModel> Parent = Node->GetParent();
+					TViewModelPtr<IOutlinerExtension> Parent = Node.AsModel()->FindAncestorOfType<IOutlinerExtension>();
 					while (Parent.IsValid())
 					{
-						TreeView->SetItemExpansion(Parent->AsShared(), true);
-						Parent = Parent->GetParent();
+						TreeView->SetItemExpansion(Parent, true);
+						Parent = Parent.AsModel()->FindAncestorOfType<IOutlinerExtension>();
 					}
 
 					TreeView->RequestScrollIntoView(Node);
@@ -6627,15 +6627,15 @@ void FSequencer::SelectNodesByPath(const TSet<FString>& NodePaths)
 
 		TSharedPtr<SOutlinerView> TreeView = SequencerWidget->GetTreeView();
 		const TSet<TWeakPtr<FViewModel>>& OutlinerSelection = GetSelection().GetSelectedOutlinerItems();
-		for (const TWeakPtr<FViewModel>& WeakNode : OutlinerSelection)
+		for (TWeakPtr<FViewModel> WeakNode : OutlinerSelection)
 		{
-			if (TSharedPtr<FViewModel> Node = WeakNode.Pin())
+			if (TViewModelPtr<IOutlinerExtension> Node = CastViewModel<IOutlinerExtension>(WeakNode.Pin()))
 			{
-				TSharedPtr<FViewModel> Parent = Node->GetParent();
+				TViewModelPtr<IOutlinerExtension> Parent = Node.AsModel()->FindAncestorOfType<IOutlinerExtension>();
 				while (Parent.IsValid())
 				{
-					TreeView->SetItemExpansion(Parent->AsShared(), true);
-					Parent = Parent->GetParent();
+					TreeView->SetItemExpansion(Parent, true);
+					Parent = Parent.AsModel()->FindAncestorOfType<IOutlinerExtension>();
 				}
 
 				TreeView->RequestScrollIntoView(Node);
@@ -6993,7 +6993,7 @@ void FSequencer::SelectByNthCategoryNode(UMovieSceneSection* Section, int Index,
 
 		if (NodesToSelect.Num() > 0)
 		{
-			SequencerWidget->GetTreeView()->RequestScrollIntoView(NodesToSelect[0]);
+			SequencerWidget->GetTreeView()->RequestScrollIntoView(CastViewModel<IOutlinerExtension>(NodesToSelect[0]));
 
 			Selection.AddToSelection(NodesToSelect);
 			Selection.RequestOutlinerNodeSelectionChangedBroadcast();
@@ -7059,7 +7059,7 @@ void FSequencer::SelectByChannels(UMovieSceneSection* Section, TArrayView<const
 		if (NodesToSelect.Num() > 0)
 		{
 			//todo hide behind preference 
-			SequencerWidget->GetTreeView()->RequestScrollIntoView(NodesToSelect[0]);
+			SequencerWidget->GetTreeView()->RequestScrollIntoView(CastViewModel<IOutlinerExtension>(NodesToSelect[0]));
 		
 			Selection.AddToSelection(NodesToSelect);
 			Selection.RequestOutlinerNodeSelectionChangedBroadcast();
@@ -7130,7 +7130,7 @@ void FSequencer::SelectByChannels(UMovieSceneSection* Section, const TArray<FNam
 
 		if (NodesToSelect.Num() > 0)
 		{
-			SequencerWidget->GetTreeView()->RequestScrollIntoView(NodesToSelect[0]);
+			SequencerWidget->GetTreeView()->RequestScrollIntoView(CastViewModel<IOutlinerExtension>(NodesToSelect[0]));
 
 			Selection.AddToSelection(NodesToSelect);
 			Selection.RequestOutlinerNodeSelectionChangedBroadcast();
diff --git a/Engine/Source/Editor/Sequencer/Private/SequencerContextMenus.cpp b/Engine/Source/Editor/Sequencer/Private/SequencerContextMenus.cpp
--- a/Engine/Source/Editor/Sequencer/Private/SequencerContextMenus.cpp
+++ b/Engine/Source/Editor/Sequencer/Private/SequencerContextMenus.cpp
@@ -625,7 +625,7 @@ void FSectionContextMenu::SelectAllKeys()
 	FSequencerSelection& Selection = Sequencer->GetSelection();
 	for (const TWeakPtr<FViewModel>& WeakModel : Selection.GetSelectedTrackAreaItems())
 	{
-		TViewModelPtr<FSectionModel> Section = WeakModel.Pin();
+		TViewModelPtr<FSectionModel> Section = CastViewModel<FSectionModel>(WeakModel.Pin());
 		if (Section)
 		{
 			UMovieSceneSection* SectionObject = Section->GetSection();
diff --git a/Engine/Source/Editor/Sequencer/Private/SequencerNodeTree.cpp b/Engine/Source/Editor/Sequencer/Private/SequencerNodeTree.cpp
--- a/Engine/Source/Editor/Sequencer/Private/SequencerNodeTree.cpp
+++ b/Engine/Source/Editor/Sequencer/Private/SequencerNodeTree.cpp
@@ -433,7 +433,7 @@ bool FSequencerNodeTree::IsNodeMute(const UE::Sequencer::TViewModelPtr<UE::Seque
 	}
 	
 	// Children should follow their parent's behavior unless told otherwise.
-	TViewModelPtr<IOutlinerExtension> ParentNode(InNode.AsModel()->GetParent());
+	TViewModelPtr<IOutlinerExtension> ParentNode = InNode.AsModel()->FindAncestorOfType<IOutlinerExtension>();
 	if (ParentNode)
 	{
 		return IsNodeMute(ParentNode);
@@ -574,7 +574,8 @@ bool FSequencerNodeTree::GetSavedPinnedState(const UE::Sequencer::FViewModel& No
 
 bool FSequencerNodeTree::IsNodeFiltered(const TSharedPtr<UE::Sequencer::FViewModel>& Node) const
 {
-	return FilteredNodes.Contains(Node);
+	using namespace UE::Sequencer;
+	return FilteredNodes.Contains(CastViewModel<IOutlinerExtension>(Node));
 }
 
 TSharedPtr<UE::Sequencer::FSectionModel> FSequencerNodeTree::GetSectionModel(const UMovieSceneSection* Section) const
@@ -988,7 +989,7 @@ bool FSequencerNodeTree::UpdateFilters()
 	// Always include the bottom spacer
 	if (FSequenceModel* SequenceModel = RootNode->CastThis<FSequenceModel>())
 	{
-		FilteredNodes.Add(SequenceModel->GetBottomSpacer());
+		FilteredNodes.Add(CastViewModelChecked<IOutlinerExtension>(SequenceModel->GetBottomSpacer()));
 	}
 
 	// Count the total number of display nodes, and update filtered state.
diff --git a/Engine/Source/Editor/Sequencer/Private/Tools/EditToolDragOperations.cpp b/Engine/Source/Editor/Sequencer/Private/Tools/EditToolDragOperations.cpp
--- a/Engine/Source/Editor/Sequencer/Private/Tools/EditToolDragOperations.cpp
+++ b/Engine/Source/Editor/Sequencer/Private/Tools/EditToolDragOperations.cpp
@@ -1126,14 +1126,14 @@ bool FMoveKeysAndSections::HandleSectionMovement(FFrameTime MouseTime, FVector2D
 	bool bRowIndexChanged = false;
 	for (UMovieSceneSection* Section : Sections)
 	{
-		TSharedPtr<FSectionModel>      SectionModel = SectionStorage->FindModelForSection(Section);
-		TSharedPtr<FViewModel> TrackModel   = SectionModel ? SectionModel->GetParentTrackModel() : nullptr;
-		if (!SectionModel || !TrackModel)
+		TSharedPtr<FSectionModel>      SectionModel  = SectionStorage->FindModelForSection(Section);
+		TViewModelPtr<ITrackExtension> TrackExtModel = SectionModel ? SectionModel->GetParentTrackModel() : nullptr;
+		if (!SectionModel || !TrackExtModel)
 		{
 			continue;
 		}
 
-		UMovieSceneTrack* Track = TrackModel->CastThis<ITrackExtension>()->GetTrack();
+		UMovieSceneTrack* Track = TrackExtModel->GetTrack();
 
 		const TArray<UMovieSceneSection*>& AllSections = Track->GetAllSections();
 
@@ -1146,7 +1146,7 @@ bool FMoveKeysAndSections::HandleSectionMovement(FFrameTime MouseTime, FVector2D
 			}
 		}
 
-		Tracks.AddUnique(TrackModel);
+		Tracks.AddUnique(TrackExtModel.AsModel());
 
 		int32 TargetRowIndex = Section->GetRowIndex();
 
@@ -1168,12 +1168,12 @@ bool FMoveKeysAndSections::HandleSectionMovement(FFrameTime MouseTime, FVector2D
 			}
 
 			// Handle sub-track and non-sub-track dragging
-			if (TrackModel->IsA<FTrackModel>())
+			if (TViewModelPtr<FTrackModel> TrackModel = TrackExtModel.ImplicitCast())
 			{
 				const int32 NumRows = FMath::Max(Section->GetRowIndex() + 1, MaxRowIndex);
 
 				// Find the total height of the track - this is necessary because tracks may contain key areas, but they will not use sub tracks unless there is more than one row
-				const FVirtualGeometry VirtualGeometry = TrackModel->CastThis<FTrackModel>()->GetVirtualGeometry();
+				const FVirtualGeometry VirtualGeometry = TrackModel->GetVirtualGeometry();
 
 				// Assume same height rows
 				const float VirtualSectionHeight = VirtualGeometry.NestedBottom - VirtualGeometry.Top;
@@ -1192,9 +1192,9 @@ bool FMoveKeysAndSections::HandleSectionMovement(FFrameTime MouseTime, FVector2D
 					TargetRowIndex = -1;
 				}
 			}
-			else if (TrackModel->IsA<FTrackRowModel>())
+			else if (TViewModelPtr<FTrackRowModel> TrackRow = TrackExtModel.ImplicitCast())
 			{
-				TSharedPtr<FTrackModel> ParentTrack = TrackModel->FindAncestorOfType<FTrackModel>();
+				TSharedPtr<FTrackModel> ParentTrack = TrackExtModel.AsModel()->FindAncestorOfType<FTrackModel>();
 				if (ensure(ParentTrack.IsValid()))
 				{
 					int32 ChildIndex = 0;
@@ -1212,8 +1212,7 @@ bool FMoveKeysAndSections::HandleSectionMovement(FFrameTime MouseTime, FVector2D
 						if (ChildIndex == 0 && (VirtualMousePos.Y <= VirtualSectionTop || LocalMousePos.Y <= 0))
 						{
 							TargetRowIndex = 0;
-							FTrackRowModel* TrackRowModel = TrackModel->CastThis<FTrackRowModel>();
-							for (TSharedPtr<FSectionModel> SectionNode : TrackRowModel->GetTrackAreaModelListAs<FSectionModel>())
+							for (TSharedPtr<FSectionModel> SectionNode : TrackRow->GetTrackAreaModelListAs<FSectionModel>())
 							{
 								if (!Sections.Contains(SectionNode->GetSection()))
 								{
diff --git a/Engine/Source/Editor/Sequencer/Public/MVVM/Views/SOutlinerTrackView.h b/Engine/Source/Editor/Sequencer/Public/MVVM/Views/SOutlinerTrackView.h
--- a/Engine/Source/Editor/Sequencer/Public/MVVM/Views/SOutlinerTrackView.h
+++ b/Engine/Source/Editor/Sequencer/Public/MVVM/Views/SOutlinerTrackView.h
@@ -2,6 +2,7 @@
 
 #pragma once
 
+#include "SequencerCoreFwd.h"
 #include "MVVM/Views/SOutlinerItemViewBase.h"
 
 class UMovieSceneTrack;
diff --git a/Engine/Source/Editor/SequencerCore/Private/MVVM/Extensions/DynamicExtensionContainer.cpp b/Engine/Source/Editor/SequencerCore/Private/MVVM/Extensions/DynamicExtensionContainer.cpp
new file mode 100644
--- /dev/null
+++ b/Engine/Source/Editor/SequencerCore/Private/MVVM/Extensions/DynamicExtensionContainer.cpp
@@ -0,0 +1,21 @@
+// Copyright Epic Games, Inc. All Rights Reserved.
+
+#include "MVVM/Extensions/DynamicExtensionContainer.h"
+
+namespace UE::Sequencer
+{
+
+const void* FDynamicExtensionContainer::CastDynamic(FViewModelTypeID Type) const
+{
+	for (const FDynamicExtensionInfo& DynamicExtension : DynamicExtensions)
+	{
+		if (const void* Result = DynamicExtension.TypeTable->Cast(&DynamicExtension.Extension.Get(), Type.GetTypeID()))
+		{
+			return Result;
+		}
+	}
+	return nullptr;
+}
+
+} // namespace UE::Sequencer
+
diff --git a/Engine/Source/Editor/SequencerCore/Private/MVVM/ICastable.cpp b/Engine/Source/Editor/SequencerCore/Private/MVVM/ICastable.cpp
--- a/Engine/Source/Editor/SequencerCore/Private/MVVM/ICastable.cpp
+++ b/Engine/Source/Editor/SequencerCore/Private/MVVM/ICastable.cpp
@@ -1,11 +1,28 @@
 // Copyright Epic Games, Inc. All Rights Reserved.
 
 #include "MVVM/ICastable.h"
+#include "MVVM/Extensions/DynamicExtensionContainer.h"
+#include "MVVM/CastableTypeTable.h"
+#include "HAL/IConsoleManager.h"
 
-namespace UE
+namespace UE::Sequencer
 {
-namespace Sequencer
+
+struct FCastableTypeInfo
 {
+	TMap<FName, void*> TypeNameToAllocation;
+
+	~FCastableTypeInfo()
+	{
+		for (TTuple<FName, void*>& Pair : TypeNameToAllocation)
+		{
+			FMemory::Free(Pair.Value);
+		}
+	}
+
+} GCastableTypeInfo;
+
+TMap<FName, TUniquePtr<FCastableTypeTable>> GTypeNameToTypeTable;
 
 void* ICastable::CastRaw(FViewModelTypeID InType)
 {
@@ -15,11 +32,62 @@ void* ICastable::CastRaw(FViewModelTypeID InType)
 
 const void* ICastable::CastRaw(FViewModelTypeID InType) const
 {
-	const void* ExtensionPtr = nullptr;
-	CastImpl(InType, ExtensionPtr);
-	return ExtensionPtr;
+	// Get the type info first to ensure that all necessary types are registered
+	const FCastableTypeTable& TypeInfo = GetTypeTable();
+
+	const void* Result = TypeInfo.Cast(this, InType.GetTypeID());
+	if (!Result && DynamicTypes)
+	{
+		Result = DynamicTypes->CastDynamic(InType);
+	}
+
+	return Result;
+}
+
+FCastableTypeTable* FCastableTypeTable::FCastableTypeTableGenerator::Commit(uint32 TypeID, FName InName)
+{
+	uint8* NewType = (uint8*)FMemory::Malloc(ByteSize(), alignof(FCastableTypeTable));
+
+	// FCastableTypeTable is a special type that uses a fixed allocation size determined by the size of the type table itself
+	// There are 2 sets of data that live on the end of the struct: the sparse bitset buckets, and the type offsets.
+	// Its layout in memory looks like this:
+	//	|												sizeof(FCastableTypeTable)										|		sizeof(uint16)*NumStaticTypes		|		sizeof(uint16)*TypeMask.NumBuckets()	|
+	//	|		int16*			|		TSparseBitSet		|		uint32		|		uint32		|		FName		|			uint16[NumStaticTypes]			|		uint16[TypeMask.NumBuckets()]			|
+	//	|	StaticTypeOffsets	|		  TypeMask			|	NumStaticTypes	|	ThisTypeID		|	ThisTypeName	|											|		 										|
+	//				\_________________________\________________________________________________________________________/										   /
+	//										   \______________________________________________TypeMask.Buckets.Storage____________________________________________/
+	FCastableTypeTable* NewTable = new (NewType) FCastableTypeTable();
+	NewTable->ThisTypeID = TypeID;
+	NewTable->ThisTypeName = InName;
+
+	uint8* DataOffset = NewType + sizeof(FCastableTypeTable);
+
+	// Sparse bitset storage lives at the end of the table
+	NewTable->TypeMask = TSparseBitSet<uint32, TFixedSparseBitSetBucketStorage<uint8>>({ DataOffset });
+	TypeMask.CopyToUnsafe(NewTable->TypeMask, TypeMask.NumBuckets());
+
+	DataOffset += TypeMask.NumBuckets()*sizeof(decltype(NewTable->TypeMask)::BucketType);
+
+	// TypeOffset information lives after the bitset storage
+	NewTable->StaticTypeOffsets = reinterpret_cast<int16*>(DataOffset);
+	NewTable->NumStaticTypes = (uint32)StaticTypeOffsets.Num();
+	FMemory::Memcpy(NewTable->StaticTypeOffsets, StaticTypeOffsets.GetData(), sizeof(int16)*StaticTypeOffsets.Num());
+
+	checkf(!GCastableTypeInfo.TypeNameToAllocation.Contains(InName), TEXT("Type name %s has already been registered. This is not supported!"), *InName.ToString());
+	GCastableTypeInfo.TypeNameToAllocation.Add(InName, NewType);
+
+	return NewTable;
+}
+
+const FCastableTypeTable* FCastableTypeTable::FindTypeByName(FName InName)
+{
+	if (void* const * TypeInfo = GCastableTypeInfo.TypeNameToAllocation.Find(InName))
+	{
+		return reinterpret_cast<const FCastableTypeTable*>(*TypeInfo);
+	}
+	return nullptr;
 }
 
-} // namespace Sequencer
-} // namespace UE
+
+} // namespace UE::Sequencer
 
diff --git a/Engine/Source/Editor/SequencerCore/Private/MVVM/ViewModelPtr.cpp b/Engine/Source/Editor/SequencerCore/Private/MVVM/ViewModelPtr.cpp
--- a/Engine/Source/Editor/SequencerCore/Private/MVVM/ViewModelPtr.cpp
+++ b/Engine/Source/Editor/SequencerCore/Private/MVVM/ViewModelPtr.cpp
@@ -2,27 +2,26 @@
 
 #include "MVVM/ViewModelPtr.h"
 
-namespace UE
-{
-namespace Sequencer
+namespace UE::Sequencer
 {
 
-FImplicitViewModelCast FViewModelPtr::ImplicitCast() const
+TViewModelConversions<FViewModel>::operator TSharedPtr<FViewModel>() const
 {
-	return FImplicitViewModelCast{Model};
+	return static_cast<const TViewModelPtr<FViewModel>*>(this)->Storage.GetModel();
 }
-
-FImplicitWeakViewModelCast FWeakViewModelPtr::ImplicitCast() const
+TViewModelConversions<FViewModel>::operator TWeakPtr<FViewModel>() const
 {
-	return FImplicitWeakViewModelCast{WeakModel};
+	return static_cast<const TViewModelPtr<FViewModel>*>(this)->Storage.GetModel();
 }
 
-FImplicitWeakViewModelPin FWeakViewModelPtr::ImplicitPin() const
+TViewModelConversions<const FViewModel>::operator TSharedPtr<const FViewModel>() const
 {
-	return FImplicitWeakViewModelPin{WeakModel};
+	return static_cast<const TViewModelPtr<const FViewModel>*>(this)->Storage.GetModel();
+}
+TViewModelConversions<const FViewModel>::operator TWeakPtr<const FViewModel>() const
+{
+	return static_cast<const TViewModelPtr<const FViewModel>*>(this)->Storage.GetModel();
 }
 
-
-} // namespace Sequencer
-} // namespace UE
+} // namespace UE::Sequencer
 
diff --git a/Engine/Source/Editor/SequencerCore/Private/MVVM/ViewModelTypeID.cpp b/Engine/Source/Editor/SequencerCore/Private/MVVM/ViewModelTypeID.cpp
--- a/Engine/Source/Editor/SequencerCore/Private/MVVM/ViewModelTypeID.cpp
+++ b/Engine/Source/Editor/SequencerCore/Private/MVVM/ViewModelTypeID.cpp
@@ -7,15 +7,10 @@ namespace UE
 namespace Sequencer
 {
 
-FViewModelTypeID FViewModelTypeID::RegisterNew()
+uint32 FViewModelTypeID::RegisterNewID()
 {
 	static uint32 ID = 0;
-	return FViewModelTypeID(ID++);
-}
-
-FViewModelTypeID FViewModelTypeID::Invalid()
-{
-	return FViewModelTypeID();
+	return ID++;
 }
 
 } // namespace Sequencer
diff --git a/Engine/Source/Editor/SequencerCore/Private/MVVM/ViewModels/OutlinerSpacer.cpp b/Engine/Source/Editor/SequencerCore/Private/MVVM/ViewModels/OutlinerSpacer.cpp
--- a/Engine/Source/Editor/SequencerCore/Private/MVVM/ViewModels/OutlinerSpacer.cpp
+++ b/Engine/Source/Editor/SequencerCore/Private/MVVM/ViewModels/OutlinerSpacer.cpp
@@ -83,7 +83,7 @@ void FOutlinerSpacer::PerformDrop(const FViewModelPtr& TargetModel, const FDragD
 	TViewModelPtr<IOutlinerDropTargetOutlinerExtension> Parent = FindAncestorOfType<IOutlinerDropTargetOutlinerExtension>();
 	if (Parent)
 	{
-		Parent->PerformDrop(this, DragDropEvent, EItemDropZone::AboveItem);
+		Parent->PerformDrop(SharedThis(this), DragDropEvent, EItemDropZone::AboveItem);
 	}
 }
 
diff --git a/Engine/Source/Editor/SequencerCore/Private/MVVM/ViewModels/OutlinerViewModel.cpp b/Engine/Source/Editor/SequencerCore/Private/MVVM/ViewModels/OutlinerViewModel.cpp
--- a/Engine/Source/Editor/SequencerCore/Private/MVVM/ViewModels/OutlinerViewModel.cpp
+++ b/Engine/Source/Editor/SequencerCore/Private/MVVM/ViewModels/OutlinerViewModel.cpp
@@ -31,7 +31,7 @@ FOutlinerViewModel::FOutlinerViewModel()
 
 void FOutlinerViewModel::Initialize(const FWeakViewModelPtr& InWeakRootDataModel)
 {
-	ensureMsgf(!WeakRootDataModel.IsValid(), TEXT("This outliner has already been initialized"));
+	ensureMsgf(!WeakRootDataModel.Pin().IsValid(), TEXT("This outliner has already been initialized"));
 	ensureMsgf(GetParent().IsValid(), TEXT("This outliner must be parented to an FEditorViewModel before being initialized"));
 	
 	WeakRootDataModel = InWeakRootDataModel;
diff --git a/Engine/Source/Editor/SequencerCore/Private/MVVM/ViewModels/ViewModel.cpp b/Engine/Source/Editor/SequencerCore/Private/MVVM/ViewModels/ViewModel.cpp
--- a/Engine/Source/Editor/SequencerCore/Private/MVVM/ViewModels/ViewModel.cpp
+++ b/Engine/Source/Editor/SequencerCore/Private/MVVM/ViewModels/ViewModel.cpp
@@ -77,6 +77,21 @@ FViewModelChildren FViewModel::GetChildList(EViewModelListType InType)
 	return Result.GetValue();
 }
 
+FViewModelPtr FViewModel::GetParent() const
+{
+	return WeakParent.Pin();
+}
+
+FViewModelPtr FViewModel::GetRoot() const
+{
+	TSharedPtr<FViewModel> CurrentItem(SharedThis(const_cast<FViewModel*>(this)));
+	while (CurrentItem->GetParent())
+	{
+		CurrentItem = CurrentItem->GetParent();
+	}
+	return CurrentItem;
+}
+
 FViewModel::ESetParentResult FViewModel::SetParentOnly(const TSharedPtr<FViewModel>& NewParent, bool bReportChanges)
 {
 	ensureAlwaysMsgf(ActiveIterationCount == 0, TEXT("Attempting to make hierarchical changes to a view model while an iteration is active"));
diff --git a/Engine/Source/Editor/SequencerCore/Private/MVVM/Views/SOutlinerItemViewBase.cpp b/Engine/Source/Editor/SequencerCore/Private/MVVM/Views/SOutlinerItemViewBase.cpp
--- a/Engine/Source/Editor/SequencerCore/Private/MVVM/Views/SOutlinerItemViewBase.cpp
+++ b/Engine/Source/Editor/SequencerCore/Private/MVVM/Views/SOutlinerItemViewBase.cpp
@@ -58,7 +58,7 @@ void SOutlinerItemViewBase::Construct(
 		const TSharedRef<ISequencerTreeViewRow>& InTableRow )
 {
 	WeakOutlinerExtension = InWeakExtension;
-	WeakRenameExtension   = InWeakExtension.ImplicitCast();
+	WeakRenameExtension   = InWeakExtension.ImplicitPin();
 	WeakEditor            = InWeakEditor;
 
 	ItemStyle              = InArgs._ItemStyle;
@@ -368,7 +368,7 @@ FOptionalSize SOutlinerItemViewBase::GetHeight() const
 
 void SOutlinerItemViewBase::OnMouseEnter(const FGeometry& MyGeometry, const FPointerEvent& MouseEvent)
 {
-	TSharedPtr<FViewModel> DataModel = WeakOutlinerExtension.Pin().AsModel();
+	TViewModelPtr<IOutlinerExtension> DataModel = WeakOutlinerExtension.Pin();
 	TSharedPtr<FEditorViewModel> Editor = WeakEditor.Pin();
 	if (DataModel && Editor)
 	{
diff --git a/Engine/Source/Editor/SequencerCore/Private/MVVM/Views/SOutlinerView.cpp b/Engine/Source/Editor/SequencerCore/Private/MVVM/Views/SOutlinerView.cpp
--- a/Engine/Source/Editor/SequencerCore/Private/MVVM/Views/SOutlinerView.cpp
+++ b/Engine/Source/Editor/SequencerCore/Private/MVVM/Views/SOutlinerView.cpp
@@ -382,7 +382,7 @@ void SOutlinerView::ForceSetSelectedItems(const TSet<TWeakPtr<FViewModel>>& InIt
 	Private_ClearSelection();
 	for (const TWeakPtr<FViewModel>& Item : InItems)
 	{
-		TWeakViewModelPtr<IOutlinerExtension> WeakItem(Item);
+		TWeakViewModelPtr<IOutlinerExtension> WeakItem = CastViewModel<IOutlinerExtension>(Item.Pin());
 		if (WeakItem.Pin())
 		{
 			Private_SetItemSelection(WeakItem, true, false);
@@ -853,7 +853,7 @@ void SOutlinerView::Refresh()
 
 		if (bIsPinned == bShowPinnedNodes)
 		{
-			RootNodes.Add(Extension.AsModel());
+			RootNodes.Add(Extension);
 		}
 	}
 
@@ -1156,7 +1156,7 @@ void SOutlinerView::CreateTrackLanesForRow(TSharedRef<SOutlinerViewRow> InRow, T
 	// We also create a lane on the appropriate parent if it doesn't exist, because we could be in a situation
 	// where we jump-scrolled directly to the given row and item without ever having had the parent in view.
 	const bool bIncludeDataModel = true;
-	for (TViewModelPtr<IOutlinerExtension> CurDataModel : InDataModel.AsModel()->GetAncestors(bIncludeDataModel))
+	for (TViewModelPtr<IOutlinerExtension> CurDataModel : InDataModel.AsModel()->GetAncestorsOfType<IOutlinerExtension>(bIncludeDataModel))
 	{
 		if (TSharedPtr<ITrackAreaExtension> TrackAreaExtension = CurDataModel.ImplicitCast())
 		{
diff --git a/Engine/Source/Editor/SequencerCore/Public/MVVM/CastableTypeTable.h b/Engine/Source/Editor/SequencerCore/Public/MVVM/CastableTypeTable.h
new file mode 100644
--- /dev/null
+++ b/Engine/Source/Editor/SequencerCore/Public/MVVM/CastableTypeTable.h
@@ -0,0 +1,249 @@
+// Copyright Epic Games, Inc. All Rights Reserved.
+
+#pragma once
+
+#include "CoreTypes.h"
+#include "HAL/Platform.h"
+#include "Templates/UnrealTemplate.h"
+#include "MVVM/SparseBitSet.h"
+#include "UObject/NameTypes.h"
+
+#include <type_traits>
+
+namespace UE::Sequencer
+{
+
+template<typename T>
+struct TViewModelTypeID;
+
+template<typename ...T>
+struct TImplements;
+
+class ICastable;
+class IDynamicExtension;
+
+struct FCastableTypeTable
+{
+	FCastableTypeTable() = default;
+
+	FCastableTypeTable(const FCastableTypeTable&) = delete;
+	void operator=(const FCastableTypeTable&) = delete;
+	FCastableTypeTable(FCastableTypeTable&&) = delete;
+	void operator=(FCastableTypeTable&&) = delete;
+
+	/**
+	 * Make a new type table for the templated type, ID and typename.
+	 * The new table is allocated using malloc and is intentionally leaked
+	 */
+	template<typename T>
+	static FCastableTypeTable* MakeTypeTable(const void* Unused, uint32 TypeID, FName InName)
+	{
+		return nullptr;
+	}
+	template<typename T>
+	static FCastableTypeTable* MakeTypeTable(const ICastable* Unused, uint32 TypeID, FName InName)
+	{
+		FCastableTypeTableGenerator Generator;
+		Generator.Generate<T>();
+		return Generator.Commit(TypeID, InName);
+	}
+	template<typename T>
+	static FCastableTypeTable* MakeTypeTable(const IDynamicExtension* Unused, uint32 TypeID, FName InName)
+	{
+		FCastableTypeTableGenerator Generator;
+		Generator.Generate<T>();
+		return Generator.Commit(TypeID, InName);
+	}
+
+	/**
+	 * Attempt to locate a type table by its C++ name
+	 */
+	SEQUENCERCORE_API static const FCastableTypeTable* FindTypeByName(FName InName);
+
+
+	/**
+	 * Use this type table to cast the specified base ptr to another type
+	 * @note BasePtr must be an instance of the type specified by this type table, or the result is undefined.
+	 *
+	 * @return BasePtr casted to the type ID specified by ToType, or nullptr if the cast failed
+	 */
+	const void* Cast(const ICastable* BasePtr, uint32 ToType) const
+	{
+		return CastImpl(BasePtr, ToType);
+	}
+
+
+	/**
+	 * Use this type table to cast the specified base ptr to another type
+	 * @note BasePtr must be an instance of the type specified by this type table, or the result is undefined.
+	 *
+	 * @return BasePtr casted to the type ID specified by ToType, or nullptr if the cast failed
+	 */
+	const void* Cast(const IDynamicExtension* BasePtr, uint32 ToType) const
+	{
+		return CastImpl(BasePtr, ToType);
+	}
+
+	/**
+	 * Retrieve the C++ type name that this type table relates to
+	 */
+	FName GetTypeName() const
+	{
+		return ThisTypeName;
+	}
+
+
+	/**
+	 * Retrieve the type ID for this type table.
+	 * Warning: This type ID is not deterministic or stable across different processes and should not be saved persistently
+	 */
+	uint32 GetTypeID() const
+	{
+		return ThisTypeID;
+	}
+
+private:
+
+	/**
+	 * Cast implementation function
+	 */
+	const void* CastImpl(const void* BasePtr, uint32 ToType) const
+	{
+		const int32 SparseTypeIndex = TypeMask.GetSparseBucketIndex(ToType);
+		if (SparseTypeIndex != INDEX_NONE)
+		{
+			return static_cast<const uint8*>(BasePtr) + StaticTypeOffsets[SparseTypeIndex];
+		}
+		return nullptr;
+	}
+
+private:
+
+	/** Utility class that generates a type table used for dynamic casting and instance of an ICastable object */ 
+	struct FCastableTypeTableGenerator
+	{
+		/**
+		 * Generate the type table from a template type
+		 */
+		template<typename T>
+		void Generate()
+		{
+			// Compute the start offset from a base ICastable to T
+			// This ensures we compute the offsets from the correct point
+			const int16 StartOffset = ComputeStartOffset<T>((T*)0);
+			AddStaticType<T>((T*)nullptr, StartOffset);
+			ConsumeImplements<T>((T*)nullptr, StartOffset);
+		}
+
+		/** Commit this generator to its final form */
+		SEQUENCERCORE_API FCastableTypeTable* Commit(uint32 TypeID, FName InName);
+
+		/**
+		 * Compute the ptr offset from one class to another in an inheritance hierarchy.
+		 * @note: This can be negative for situations where multiple inheritance would put
+		 *        the FromType class after the ToType in memory.
+		 */
+		template<typename FromType, typename ToType>
+		static int16 ClassOffsetFrom()
+		{
+			alignas(FromType) uint8 Buffer[sizeof(FromType)];
+
+			FromType* From = reinterpret_cast<FromType*>(Buffer);
+			ToType*   To   = static_cast<ToType*>(From);
+
+			const int32 Offset = static_cast<int32>((uint8*)To - (uint8*)From);
+			checkf(Offset >= std::numeric_limits<int16>::lowest() && Offset <= std::numeric_limits<int16>::max(), TEXT("Class offset is too large to be supported by ICastable"));
+			return static_cast<int16>(Offset);
+		}
+
+		/**
+		 * Compute the start offset from an ICastable type to the root of T
+		 * This offset is added to all other class offsets that are computed from T
+		 * so that we can always cast from an ICastable* to the concrete types
+		 */
+		template<typename T>
+		static int16 ComputeStartOffset(const ICastable*)
+		{
+			return ClassOffsetFrom<ICastable, T>();
+		}
+
+		/**
+		 * Compute the start offset from an IDynamicExtension type to the root of T
+		 * This offset is added to all other class offsets that are computed from T
+		 * so that we can always cast from an IDynamicExtension* to the concrete types
+		 */
+		template<typename T>
+		static int16 ComputeStartOffset(const IDynamicExtension*)
+		{
+			return ClassOffsetFrom<IDynamicExtension, T>();
+		}
+
+		uint32 ByteSize() const
+		{
+			return sizeof(FCastableTypeTable) + sizeof(int16)*StaticTypeOffsets.Num() + sizeof(uint8)*StaticTypeOffsets.Num();
+		}
+
+		template<typename StartClass, typename T, typename U = typename T::Implements>
+		void ConsumeImplements(T*, int16 StartOffset)
+		{
+			ConsumeImplementsImpl<StartClass>((typename T::Implements*)nullptr, StartOffset);
+		}
+
+		template<typename StartClass>
+		void ConsumeImplements(...)
+		{
+		}
+
+		template<typename StartClass, typename ...Types>
+		void ConsumeImplementsImpl(TImplements<Types...>*, int16 StartOffset)
+		{
+			(AddStaticType<StartClass>((Types*)nullptr, StartOffset), ...);
+			(ConsumeImplements<StartClass>((Types*)nullptr, StartOffset), ...);
+		}
+
+		template<typename StartClass, typename T>
+		void ConsumeImplements(void*, int16 StartOffset)
+		{
+		}
+
+		template<typename StartClass, typename T, typename U = decltype(T::ID)>
+		void AddStaticType(T*, int16 StartOffset)
+		{
+			TViewModelTypeID<T>& ID = T::ID;
+
+			int16 OffsetFromBase = StartOffset + ClassOffsetFrom<StartClass, T>();
+			if (TypeMask.SetBit(ID.GetTypeID()) == ESparseBitSetBitResult::NewlySet)
+			{
+				int32 SparseIndex = TypeMask.GetSparseBucketIndex(ID.GetTypeID());
+				StaticTypeOffsets.Insert(OffsetFromBase, SparseIndex);
+			}
+		}
+
+		template<typename StartClass>
+		constexpr void AddStaticType(void*, int16 StartOffset)
+		{
+		}
+
+	private:
+
+		/* Accumulator sparse bitset of types that are implemented on the type that Generate was called with */
+		TSparseBitSet<uint32> TypeMask;
+		/* Typeoffsets from T to the ptr for each type within TypeMask */
+		TArray<int16> StaticTypeOffsets;
+	};
+
+private:
+
+	/** Ptr to the static type offsets allocated at the end of this struct. Size = NumStaticTypes */
+	int16* StaticTypeOffsets;
+	/** Fixed sparse bitset of types that this table contains. Storage is allocated at the end of this struct. */
+	TSparseBitSet<uint32, TFixedSparseBitSetBucketStorage<uint8>> TypeMask;
+	/** The number of static types implemented within this type table. */
+	uint32 NumStaticTypes;
+	/** Unique static type ID for this type. */
+	uint32 ThisTypeID;
+	/** The name of the C++ type that this table relates to. */
+	FName ThisTypeName;
+};
+
+} // namespace UE::Sequencer
\ No newline at end of file
diff --git a/Engine/Source/Editor/SequencerCore/Public/MVVM/Extensions/DynamicExtensionContainer.h b/Engine/Source/Editor/SequencerCore/Public/MVVM/Extensions/DynamicExtensionContainer.h
--- a/Engine/Source/Editor/SequencerCore/Public/MVVM/Extensions/DynamicExtensionContainer.h
+++ b/Engine/Source/Editor/SequencerCore/Public/MVVM/Extensions/DynamicExtensionContainer.h
@@ -15,6 +15,7 @@ namespace UE
 namespace Sequencer
 {
 
+struct FCastableTypeTable;
 class FDynamicExtensionContainer;
 class FViewModel;
 
@@ -42,7 +43,7 @@ DECLARE_DELEGATE_OneParam(FDynamicExtensionCallback, TSharedRef<FViewModel>);
 /**
  * Base class for supporting dynamic extensions that participate in dynamic casting.
  */
-class SEQUENCERCORE_API FDynamicExtensionContainer
+class FDynamicExtensionContainer
 {
 public:
 
@@ -55,7 +56,7 @@ public:
 
 public:
 
-	UE_SEQUENCER_DECLARE_VIEW_MODEL_TYPE_ID(FDynamicExtensionContainer)
+	UE_SEQUENCER_DECLARE_VIEW_MODEL_TYPE_ID_API(SEQUENCERCORE_API, FDynamicExtensionContainer)
 
 	template<typename T>
 	const T* CastDynamic() const
@@ -83,17 +84,7 @@ public:
 		return const_cast<T*>(const_cast<const FDynamicExtensionContainer*>(this)->CastDynamicChecked<T>());
 	}
 
-	const void* CastDynamic(FViewModelTypeID Type) const
-	{
-		for (const TInlineValue<FDynamicExtensionInfo>& DynamicExtension : DynamicExtensions)
-		{
-			if (const void* Result = DynamicExtension->CastDynamic(Type))
-			{
-				return Result;
-			}
-		}
-		return nullptr;
-	}
+	SEQUENCERCORE_API const void* CastDynamic(FViewModelTypeID Type) const;
 
 	const void* CastDynamicChecked(FViewModelTypeID Type) const
 	{
@@ -125,90 +116,37 @@ protected:
 			return static_cast<T&>(*Existing);
 		}
 
-		static_assert(sizeof(FDynamicExtensionInfo) == sizeof(TDynamicExtensionInfo<T>), "Typed extension wrapper can't contain extra data.");
-
 		TSharedRef<T> NewExtension = MakeShared<T>(Forward<InArgTypes>(Args)...);
-		DynamicExtensions.Add(TDynamicExtensionInfo<T>(NewExtension));
+		DynamicExtensions.Emplace(FDynamicExtensionInfo{ T::ID.GetTypeTable(), NewExtension });
 		NewExtension->OnCreated(InOwner);
 		return NewExtension.Get();
 	}
 
 	void PostInitializeExtensions()
 	{
-		for (const TInlineValue<FDynamicExtensionInfo>& DynamicExtension : DynamicExtensions)
+		for (const FDynamicExtensionInfo& DynamicExtension : DynamicExtensions)
 		{
-			check(DynamicExtension.IsValid());
-			DynamicExtension->Extension->OnPostInitialize();
+			DynamicExtension.Extension->OnPostInitialize();
 		}
 	}
 
 	void ReinitializeExtensions()
 	{
-		for (const TInlineValue<FDynamicExtensionInfo>& DynamicExtension : DynamicExtensions)
+		for (const FDynamicExtensionInfo& DynamicExtension : DynamicExtensions)
 		{
-			check(DynamicExtension.IsValid());
-			DynamicExtension->Extension->OnReinitialize();
+			DynamicExtension.Extension->OnReinitialize();
 		}
 	}
 
-public:
-
-	struct Implements{};
-
 private:
 
 	struct FDynamicExtensionInfo
 	{
+		FCastableTypeTable* TypeTable;
 		TSharedRef<IDynamicExtension> Extension;
-		FDynamicExtensionInfo(TSharedRef<IDynamicExtension> InExtension) : Extension(InExtension) {}
-		virtual ~FDynamicExtensionInfo() {}
-		virtual const void* CastDynamic(FViewModelTypeID Type) const { return nullptr; }
 	};
 
-	template<typename T>
-	struct TDynamicExtensionInfo : FDynamicExtensionInfo
-	{
-		TDynamicExtensionInfo(TSharedRef<T> InExtension) 
-			: FDynamicExtensionInfo(StaticCastSharedRef<IDynamicExtension>(InExtension))
-		{
-		}
-		virtual const void* CastDynamic(FViewModelTypeID Type) const
-		{
-			const void* Result = nullptr;
-			TSharedRef<T> TypedExtension = StaticCastSharedRef<T>(Extension);
-			FCastableBoilerplate::CastImplementation(&TypedExtension.Get(), Type, Result);
-			return Result;
-		}
-	};
-
-	TArray<TInlineValue<FDynamicExtensionInfo>> DynamicExtensions;
-};
-
-template<>
-struct TCompositeCast<FDynamicExtensionContainer::Implements>
-{
-	static void Apply(const FDynamicExtensionContainer* This, FViewModelTypeID ToType, const void*& OutResult)
-	{
-		if (OutResult == nullptr)
-		{
-			if (const void* Result = FCastableBoilerplate::DirectCast(This, ToType))
-			{
-				OutResult = Result;
-				return;
-			}
-
-			if (const void* Result = This->CastDynamic(ToType))
-			{
-				OutResult = Result;
-			}
-		}
-	}
-
-	template<typename T>
-	static bool IsAny(const ICastable* This)
-	{
-		return T::ID == FDynamicExtensionContainer::ID;
-	}
+	TArray<FDynamicExtensionInfo> DynamicExtensions;
 };
 
 } // namespace Sequencer
diff --git a/Engine/Source/Editor/SequencerCore/Public/MVVM/Extensions/IRecyclableExtension.h b/Engine/Source/Editor/SequencerCore/Public/MVVM/Extensions/IRecyclableExtension.h
--- a/Engine/Source/Editor/SequencerCore/Public/MVVM/Extensions/IRecyclableExtension.h
+++ b/Engine/Source/Editor/SequencerCore/Public/MVVM/Extensions/IRecyclableExtension.h
@@ -3,6 +3,7 @@
 #pragma once
 
 #include "CoreTypes.h"
+#include "SequencerCoreFwd.h"
 #include "MVVM/ViewModelTypeID.h"
 
 namespace UE
@@ -10,8 +11,6 @@ namespace UE
 namespace Sequencer
 {
 
-template<typename T> struct TViewModelPtr;
-
 class SEQUENCERCORE_API IRecyclableExtension
 {
 public:
diff --git a/Engine/Source/Editor/SequencerCore/Public/MVVM/ICastable.h b/Engine/Source/Editor/SequencerCore/Public/MVVM/ICastable.h
--- a/Engine/Source/Editor/SequencerCore/Public/MVVM/ICastable.h
+++ b/Engine/Source/Editor/SequencerCore/Public/MVVM/ICastable.h
@@ -8,22 +8,20 @@
 #include "Templates/Invoke.h"
 #include "Templates/SharedPointer.h"
 #include "Templates/UnrealTemplate.h"
+#include "MVVM/CastableTypeTable.h"
 
 #include <type_traits>
 
-namespace UE
-{
-namespace Sequencer
+namespace UE::Sequencer
 {
 
-struct FCastableBoilerplate;
+struct FCastableTypeTable;
+class FDynamicExtensionContainer;
 
 class SEQUENCERCORE_API ICastable
 {
 public:
 
-	virtual ~ICastable(){}
-
 	template<typename RetType, typename InterfaceType, typename ...ArgTypes, typename ...ParamTypes>
 	RetType CastedCall(RetType(InterfaceType::*MemberFunc)(ArgTypes...), ParamTypes&&... InArgs)
 	{
@@ -67,9 +65,6 @@ public:
 		return CastThis<T>() != nullptr;
 	}
 
-	template<typename T>
-	bool IsAny() const;
-
 	void* CastRaw(FViewModelTypeID InType);
 	const void* CastRaw(FViewModelTypeID InType) const;
 
@@ -101,117 +96,33 @@ public:
 		return Result;
 	}
 
-protected:
-
-	// Usually implemented through UE_SEQUENCER_DECLARE_CASTABLE
-	virtual void CastImpl(FViewModelTypeID Type, const void*& OutPtr) const = 0;
-};
-
-template<typename ...T>
-struct TImplements;
-
-#define UE_SEQUENCER_DECLARE_CASTABLE(ThisType, ...)                                             \
-	UE_SEQUENCER_DECLARE_VIEW_MODEL_TYPE_ID(ThisType);                                           \
-	using Implements = TImplements<__VA_ARGS__>;                                                 \
-	virtual void CastImpl(FViewModelTypeID ToType, const void*& OutPtr) const override  \
-	{                                                                                            \
-		FCastableBoilerplate::CastImplementation(this, ToType, OutPtr);                          \
-	}
-
-#define UE_SEQUENCER_DEFINE_CASTABLE(ThisType)                                                   \
-	UE_SEQUENCER_DEFINE_VIEW_MODEL_TYPE_ID(ThisType);
-
-
-template<typename ImplementsType>
-struct TCompositeCast;
-
-struct FCastableBoilerplate
-{
-	template<typename T>
-	struct THasImplements
-	{
-		template<typename U, typename V = typename U::Implements>
-		static uint8 HasImplements(const U* Input);
-		static uint16 HasImplements(const void* Input);
-
-		static constexpr bool Value = sizeof(HasImplements((const T*)0)) == 1;
-	};
-
-	template<typename T>
-	FORCENOINLINE static void CastImplementation(const T* This, FViewModelTypeID ToType, const void*& OutResult)
+	const FCastableTypeTable& GetTypeTable() const
 	{
-		if (const void* Result = DirectCast(This, ToType))
-		{
-			OutResult = Result;
-		}
-		else
-		{
-			CompositeCast(This, ToType, OutResult);
-		}
+		// This should never trip - it should always be initialzed on construction by a FAutoRegisterTypeTable member
+		checkSlow(TypeTable);
+		return *TypeTable;
 	}
 
-	template<typename T>
-	FORCENOINLINE static bool IsAnyImplementation(const ICastable* This)
-	{
-		const void* Result = This->CastRaw(T::ID);
-		if (!Result)
-		{
-			return CompositeIsAny<T>(This);
-		}
-	}
+protected:
 
-	template<typename T, typename U = decltype(T::ID)>
-	FORCENOINLINE static const void* DirectCast(const T* Input, FViewModelTypeID ToType)
+	/** Struct that assigns an ICastable instance's TypeTable. See UE_SEQUENCER_DECLARE_CASTABLE for usage. */
+	struct FAutoRegisterTypeTable
 	{
-		if (ToType == T::ID)
+		template<typename T>
+		FAutoRegisterTypeTable(T* InThis)
 		{
-			return Input;
+			InThis->ID.Register();
+			const_cast<const FCastableTypeTable*&>(InThis->TypeTable) = InThis->ID.GetTypeTable();
 		}
-		return nullptr;
-	}
-
-	FORCENOINLINE static const void* DirectCast(const void* Input, FViewModelTypeID ToType)
-	{
-		return nullptr;
-	}
-
-	template<typename T>
-	FORCENOINLINE static std::enable_if_t<THasImplements<T>::Value> CompositeCast(const T* Input, FViewModelTypeID ToType, const void*& OutResult)
-	{
-		TCompositeCast<typename T::Implements>::Apply(Input, ToType, OutResult);
-	}
-
-	FORCENOINLINE static void CompositeCast(const void* Input, FViewModelTypeID ToType, const void*& OutResult)
-	{
-	}
+	};
 
-	template<typename T>
-	FORCENOINLINE static std::enable_if_t<THasImplements<T>::Value, bool> CompositeIsAny(const ICastable* This)
-	{
-		return TCompositeCast<typename T::Implements>::IsAny(This);
-	}
+private:
 
-	template<typename T>
-	FORCENOINLINE static std::enable_if_t<!THasImplements<T>::Value, bool> CompositeIsAny(const ICastable* This)
-	{
-	}
-};
+	friend FAutoRegisterTypeTable;
+	const FCastableTypeTable* const TypeTable = nullptr;
 
-template<typename ...CastableTo>
-struct TCompositeCast<TImplements<CastableTo...>>
-{
-	template<typename T>
-	FORCENOINLINE static void Apply(const T* This, FViewModelTypeID ToType, const void*& OutResult)
-	{
-		(FCastableBoilerplate::CastImplementation(static_cast<const CastableTo*>(This), ToType, OutResult), ...);
-	}
-	template<typename T>
-	FORCENOINLINE static bool IsAny(const ICastable* This)
-	{
-		bool bIsAny = false;
-		(( bIsAny |= FCastableBoilerplate::IsAnyImplementation<CastableTo>(This) ), ...);
-		return bIsAny;
-	}
+protected:
+	FDynamicExtensionContainer* DynamicTypes = nullptr;
 };
 
 template<typename ...T>
@@ -219,13 +130,13 @@ struct TImplements
 {
 };
 
+#define UE_SEQUENCER_DECLARE_CASTABLE(ThisType, ...)                                             \
+	UE_SEQUENCER_DECLARE_VIEW_MODEL_TYPE_ID(ThisType);                                           \
+	using Implements = TImplements<__VA_ARGS__>;                                                 \
+	UE_NO_UNIQUE_ADDRESS FAutoRegisterTypeTable AutoRegisterTypeTable = FAutoRegisterTypeTable(static_cast<ThisType*>(this)); \
 
-template<typename T>
-bool ICastable::IsAny() const
-{
-	return FCastableBoilerplate::template IsAnyImplementation<T>(this);
-}
+#define UE_SEQUENCER_DEFINE_CASTABLE(ThisType)                                                   \
+	UE_SEQUENCER_DEFINE_VIEW_MODEL_TYPE_ID(ThisType);
 
-} // namespace Sequencer
-} // namespace UE
 
+} // namespace UE::Sequencer
\ No newline at end of file
diff --git a/Engine/Source/Editor/SequencerCore/Public/MVVM/SparseBitSet.h b/Engine/Source/Editor/SequencerCore/Public/MVVM/SparseBitSet.h
new file mode 100644
--- /dev/null
+++ b/Engine/Source/Editor/SequencerCore/Public/MVVM/SparseBitSet.h
@@ -0,0 +1,236 @@
+// Copyright Epic Games, Inc. All Rights Reserved.
+
+#pragma once
+
+#include "CoreTypes.h"
+#include "Templates/UnrealTemplate.h"
+#include "Containers/Array.h"
+
+namespace UE::Sequencer
+{
+
+template<typename T>
+struct TDynamicSparseBitSetBucketStorage;
+
+template<typename T>
+struct TFixedSparseBitSetBucketStorage;
+
+template<typename HashType, typename BucketStorage = TDynamicSparseBitSetBucketStorage<uint8>>
+struct TSparseBitSet;
+
+
+enum class ESparseBitSetBitResult
+{
+	NewlySet,
+	AlreadySet,
+};
+
+/**
+ * NOTE: This class is currently considered internal only, and should only be used by engine code.
+ * A sparse bitset comprising a hash of integer indexes with set bits, and a sparse array of unsigned integers (referred to as buckets) whose width is defined by the storage.
+ *
+ * The maximum size bitfield that is representible by this bitset is defined as sizeof(HashType)*sizeof(BucketStorage::BucketType). For example, a 64 bit hash with 32 bit buckets
+ * can represent a bitfield of upto 2048 bits.
+ *
+ * The hash allows for empty buckets to be completely omitted from memory, and affords very fast comparison for buckets that have no set bits.
+ * This container is specialized for relatively large bitfields that have small numbers of set bits (ie, needles in haystacks) as they will provide the best memory vs CPU tradeoffs.
+ */
+template<typename HashType, typename BucketStorage>
+struct TSparseBitSet
+{
+	using BucketType = typename BucketStorage::BucketType;
+	static constexpr uint32 HashSize    = sizeof(HashType)*8;
+	static constexpr uint32 BucketSize  = sizeof(typename BucketStorage::BucketType)*8;
+	static constexpr uint32 MaxBitIndex = HashSize * BucketSize;
+
+	TSparseBitSet()
+		: BucketHash(0)
+	{}
+
+	TSparseBitSet(BucketStorage&& InBucketStorage)
+		: Buckets(MoveTemp(InBucketStorage))
+		, BucketHash(0)
+	{}
+
+	template<typename OtherHashType, typename OtherStorageType>
+	void CopyTo(TSparseBitSet<OtherHashType, OtherStorageType>& Other) const
+	{
+		static_assert(TSparseBitSet<OtherHashType, OtherStorageType>::BucketSize == BucketSize, "Cannot copy sparse bitsets of different bucket sizes");
+
+		// Copy the buckets
+		const uint32 NumBuckets = FMath::CountBits(Other.BucketHash);
+		Other.Buckets.SetNum(NumBuckets);
+		CopyToUnsafe(Other, NumBuckets);
+	}
+
+	/** Copy this bitset to another without resizing the destination's bucket storage. Bucket storage must be >= this size. */
+	template<typename OtherHashType, typename OtherStorageType>
+	void CopyToUnsafe(TSparseBitSet<OtherHashType, OtherStorageType>& Other, uint32 OtherBucketCapacity) const
+	{
+		static_assert(TSparseBitSet<OtherHashType, OtherStorageType>::BucketSize == BucketSize, "Cannot copy sparse bitsets of different bucket sizes");
+
+		const uint32 ThisNumBuckets = this->NumBuckets();
+		checkf(OtherBucketCapacity >= ThisNumBuckets, TEXT("Attempting to copy a sparse bitset without enough capacity in the destination (%d, required %d)"), OtherBucketCapacity, ThisNumBuckets);
+
+		// Copy the hash
+		Other.BucketHash = this->BucketHash;
+
+		// Copy the buckets
+		FMemory::Memcpy(Other.Buckets.GetData(), this->Buckets.Storage.GetData(), sizeof(typename BucketStorage::BucketType)*ThisNumBuckets);
+	}
+
+	/**
+	 * Count the number of buckets in this bitset
+	 */
+	uint32 NumBuckets() const
+	{
+		return FMath::CountBits(this->BucketHash);
+	}
+
+	/**
+	 * Set the bit at the specified index.
+	 * Any bits between Num and BitIndex will be considered 0.
+	 *
+	 * @return true if the bit was previously considered 0 and is now set, false if it was already set.
+	 */
+	ESparseBitSetBitResult SetBit(uint32 BitIndex)
+	{
+		CheckIndex(BitIndex);
+
+		FBitOffsets Offsets(BucketHash, BitIndex);
+
+		// Do we need to add a new bucket?
+		if ( (BucketHash & Offsets.HashBit) == 0)
+		{
+			BucketHash |= Offsets.HashBit;
+			Buckets.Insert(Offsets.BitMaskWithinBucket, Offsets.BucketIndex);
+			return ESparseBitSetBitResult::NewlySet;
+		}
+		else if ((Buckets.Get(Offsets.BucketIndex) & Offsets.BitMaskWithinBucket) == 0)
+		{
+			Buckets.Get(Offsets.BucketIndex) |= Offsets.BitMaskWithinBucket;
+			return ESparseBitSetBitResult::NewlySet;
+		}
+
+		return ESparseBitSetBitResult::AlreadySet;
+	}
+
+	/**
+	 * Check whether the specified bit index is set
+	 */
+	bool IsBitSet(uint32 BitIndex) const
+	{
+		CheckIndex(BitIndex);
+
+		const uint32 Hash    = BitIndex / BucketSize;
+		const uint32 HashBit = (1 << Hash);
+		if (BucketHash & HashBit)
+		{
+			const uint32     BucketIndex  = FMath::CountBits(BucketHash & (HashBit-1));
+			const uint32     ThisBitIndex = (BitIndex-BucketSize*Hash);
+			const BucketType ThisBitMask  = BucketType(1u) << ThisBitIndex;
+
+			return Buckets.Get(BucketIndex) & ThisBitMask;
+		}
+		return false;
+	}
+
+	/**
+	 * Get the sparse bucket index of the specified bit
+	 */
+	int32 GetSparseBucketIndex(uint32 BitIndex) const
+	{
+		CheckIndex(BitIndex);
+
+		const uint32 Hash    = BitIndex / BucketSize;
+		const uint32 HashBit = (1 << Hash);
+		if (BucketHash & HashBit)
+		{
+			uint32 BucketIndex = FMath::CountBits(BucketHash & (HashBit-1));
+
+			const uint32 ThisBitIndex = (BitIndex-BucketSize*Hash);
+			const BucketType ThisBitMask = BucketType(1u) << ThisBitIndex;
+
+			BucketType ThisBucket = Buckets.Get(BucketIndex);
+			if (ThisBucket & ThisBitMask)
+			{
+				// Compute the offset
+				int32 SparseIndex = FMath::CountBits(ThisBucket & (ThisBitMask-1));
+
+				// Count all the preceding buckets to find the final sparse index
+				while (BucketIndex > 0)
+				{
+					--BucketIndex;
+					SparseIndex += FMath::CountBits(Buckets.Get(BucketIndex));
+				}
+				return SparseIndex;
+			}
+		}
+		return INDEX_NONE;
+	}
+
+private:
+
+	template<typename, typename>
+	friend struct TSparseBitSet;
+
+	FORCEINLINE void CheckIndex(uint32 BitIndex) const
+	{
+		checkfSlow(BitIndex < MaxBitIndex, TEXT("Invalid index (%d) specified for a sparse bitset of maximum size (%d)"), BitIndex, MaxBitIndex);
+	}
+
+	struct FBitOffsets
+	{
+		uint32 HashBit;
+		int32  BucketIndex;
+		BucketType BitMaskWithinBucket;
+		FBitOffsets(HashType InBucketHash, uint32 BitIndex)
+		{
+			const uint32 Hash = BitIndex / BucketSize;
+			HashBit = (1 << Hash);
+
+			BucketIndex = FMath::CountBits(InBucketHash & (HashBit-1));
+
+			const uint32 ThisBitIndex = (BitIndex-BucketSize*Hash);
+			BitMaskWithinBucket = BucketType(1u) << ThisBitIndex;
+		}
+	};
+
+	BucketStorage Buckets;
+	HashType      BucketHash;
+};
+
+
+
+template<typename T>
+struct TDynamicSparseBitSetBucketStorage
+{
+	using BucketType = T;
+
+	TArray<BucketType, TInlineAllocator<8>> Storage;
+
+	void Insert(BucketType InitialValue, int32 Index)
+	{
+		Storage.Insert(InitialValue, Index);
+	}
+
+	BucketType* GetData()              { return Storage.GetData(); }
+
+	BucketType& Get(int32 Index)       { return Storage[Index]; }
+	BucketType  Get(int32 Index) const { return Storage[Index]; }
+};
+
+template<typename T>
+struct TFixedSparseBitSetBucketStorage
+{
+	using BucketType = T;
+
+	BucketType* Storage = nullptr;
+
+	BucketType* GetData()              { return Storage; }
+
+	BucketType& Get(int32 Index)       { return Storage[Index]; }
+	BucketType  Get(int32 Index) const { return Storage[Index]; }
+};
+
+} // namespace UE::Sequencer
\ No newline at end of file
diff --git a/Engine/Source/Editor/SequencerCore/Public/MVVM/ViewModelPtr.h b/Engine/Source/Editor/SequencerCore/Public/MVVM/ViewModelPtr.h
--- a/Engine/Source/Editor/SequencerCore/Public/MVVM/ViewModelPtr.h
+++ b/Engine/Source/Editor/SequencerCore/Public/MVVM/ViewModelPtr.h
@@ -11,14 +11,81 @@
 #include "Templates/SharedPointer.h"
 #include "Templates/TypeHash.h"
 
-namespace UE
+namespace UE::Sequencer
 {
-namespace Sequencer
+
+/**
+ * Storage selector implemented as a type alias to defer instantiation and allow use with fwd declared types
+ */
+template<typename T>
+struct TViewModelPtrStorage
+{
+	/*~ Define either 'FViewModel' or 'const FViewModel' based on the constness of T */
+	using ViewModelType = typename std::conditional_t<std::is_const_v<T>, const FViewModel, FViewModel>;
+
+	FORCEINLINE T* Get() const
+	{
+		return Extension;
+	}
+	FORCEINLINE const TSharedPtr<ViewModelType>& GetModel() const
+	{
+		return SharedModel;
+	}
+
+	void Set(TSharedPtr<ViewModelType> InModel, T* ExtensionPtr)
+	{
+		if (ExtensionPtr)
+		{
+			SharedModel = InModel;
+			Extension = ExtensionPtr;
+		}
+		else
+		{
+			SharedModel = nullptr;
+			Extension = nullptr;
+		}
+	}
+
+	TSharedPtr<ViewModelType> SharedModel;
+	T* Extension = nullptr;
+};
+
+template<typename T>
+struct TViewModelConversions
 {
-struct FImplicitViewModelCast;
-struct FImplicitWeakViewModelCast;
-struct FImplicitWeakViewModelPin;
+	/*~ Define either 'FViewModel' or 'const FViewModel' based on the constness of T */
+	using ViewModelType = typename std::conditional_t<std::is_const_v<T>, const FViewModel, FViewModel>;
 
+	operator TSharedPtr<ViewModelType>() const
+	{
+		return static_cast<const TViewModelPtr<T>*>(this)->Storage.GetModel();
+	}
+	operator TSharedPtr<T>() const
+	{
+		const TViewModelPtr<T>* This = static_cast<const TViewModelPtr<T>*>(this);
+		return TSharedPtr<T>(This->Storage.SharedModel, This->Storage.Extension);
+	}
+	operator TWeakPtr<ViewModelType>() const
+	{
+		return TSharedPtr<ViewModelType>(*this);
+	}
+	operator TWeakPtr<T>() const
+{
+		return TSharedPtr<T>(*this);
+	}
+};
+template<>
+struct TViewModelConversions<FViewModel>
+{
+	SEQUENCERCORE_API operator TSharedPtr<FViewModel>() const;
+	SEQUENCERCORE_API operator TWeakPtr<FViewModel>() const;
+};
+template<>
+struct TViewModelConversions<const FViewModel>
+{
+	SEQUENCERCORE_API operator TSharedPtr<const FViewModel>() const;
+	SEQUENCERCORE_API operator TWeakPtr<const FViewModel>() const;
+};
 
 /**
  * Pointer type that wraps a shared view model pointer to provide implicit dynamic_cast style
@@ -26,64 +93,107 @@ struct FImplicitWeakViewModelPin;
  *    a proxy object so as to make the 'implicit' conversion somewhat explicit:
  *         FViewModelPtr ModelPtr = GetViewModel();
  *         TSharedPtr<IDraggableOutlinerExtension>   Draggable = ModelPtr.ImplicitCast();
- *         TSharedPtr<IDroppableExtension>   Droppable = ModelPtr.ImplicitCast();
+ *         TSharedPtr<IDroppableExtension>           Droppable = ModelPtr.ImplicitCast();
  */
-struct FViewModelPtr
+template<typename T>
+struct TViewModelPtr : TViewModelConversions<T>
 {
-	FViewModelPtr() = default;
-	FViewModelPtr(const FViewModelPtr&) = default;
-	FViewModelPtr& operator=(const FViewModelPtr&) = default;
-	FViewModelPtr(FViewModelPtr&&) = default;
-	FViewModelPtr& operator=(FViewModelPtr&&) = default;
+	using ViewModelType = typename std::conditional_t<std::is_const_v<T>, const FViewModel, FViewModel>;
+
+	TViewModelPtr() = default;
+	TViewModelPtr(const TViewModelPtr&) = default;
+	TViewModelPtr& operator=(const TViewModelPtr&) = default;
+	TViewModelPtr(TViewModelPtr&&) = default;
+	TViewModelPtr& operator=(TViewModelPtr&&) = default;
+
+	TViewModelPtr(TSharedPtr<ViewModelType> InModel, T* InExtension)
+	{
+		if (InExtension)
+		{
+			Storage.Set(InModel, InExtension);
+		}
+	}
+
+	template<typename OtherType>
+	TViewModelPtr(const TViewModelPtr<OtherType>& In)
+	{
+		if constexpr (TPointerIsConvertibleFromTo<OtherType, T>::Value)
+		{
+			Storage.Set(In.Storage.SharedModel, In.Storage.Get());
+		}
+		else
+		{
+			Storage.Set(In.Storage.SharedModel, In.Storage.SharedModel.Get());
+		}
+	}
+
+	template<typename OtherType>
+	TViewModelPtr& operator=(const TViewModelPtr<OtherType>& In)
+	{
+		if constexpr (TPointerIsConvertibleFromTo<OtherType, T>::Value)
+		{
+			Storage.Set(In.Storage.SharedModel, In.Storage.Get());
+		}
+		else
+		{
+			Storage.Set(In.Storage.SharedModel, In.Storage.SharedModel.Get());
+		}
+		return *this;
+	}
 
 	/** Construction from a nullptr */
-	FViewModelPtr(nullptr_t)
-		: Model(nullptr)
+	TViewModelPtr(nullptr_t)
 	{}
+
 	/** Assignment from a nullptr */
-	FViewModelPtr& operator=(nullptr_t)
+	TViewModelPtr& operator=(nullptr_t)
 	{
-		Model = nullptr;
+		Storage.Set(nullptr, nullptr);
 		return *this;
 	}
 
 	/** Construction from a shared ptr */
 	template<typename ModelType>
-	FViewModelPtr(const TSharedPtr<ModelType>& InModel)
-		: Model(InModel)
-	{}
+	TViewModelPtr(const TSharedPtr<ModelType>& InModel)
+	{
+		static_assert(TPointerIsConvertibleFromTo<ModelType, T>::Value && TPointerIsConvertibleFromTo<ModelType, ViewModelType>::Value,
+			"Construction from shared pointers is only supported for related types. Please use CastViewModel<> to cast to the correct type.");
+
+		Storage.Set(InModel, InModel.Get());
+	}
 	/** Assignment from a shared ptr */
 	template<typename ModelType>
-	FViewModelPtr& operator=(const TSharedPtr<ModelType>& InModel)
+	TViewModelPtr& operator=(const TSharedPtr<ModelType>& InModel)
 	{
-		Model = InModel;
+		static_assert(TPointerIsConvertibleFromTo<ModelType, T>::Value && TPointerIsConvertibleFromTo<ModelType, ViewModelType>::Value,
+			"Construction from shared pointers is only supported for related types. Please use CastViewModel<> to cast to the correct type.");
+
+		Storage.Set(InModel, InModel.Get());
 		return *this;
 	}
 
 	/** Construction from a shared ptr */
 	template<typename ModelType>
-	FViewModelPtr(const TSharedRef<ModelType>& InModel)
-		: Model(InModel)
-	{}
-	/** Assignment from a shared ptr */
-	template<typename ModelType>
-	FViewModelPtr& operator=(const TSharedRef<ModelType>& InModel)
+	TViewModelPtr(const TSharedRef<ModelType>& InModel)
 	{
-		Model = InModel;
-		return *this;
+		static_assert(TPointerIsConvertibleFromTo<ModelType, T>::Value && TPointerIsConvertibleFromTo<ModelType, ViewModelType>::Value,
+			"Construction from shared pointers is only supported for related types. Please use CastViewModel<> to cast to the correct type.");
+
+		Storage.Set(TSharedPtr<ModelType>(InModel), &InModel.Get());
 	}
 
-	/** Construction from a raw ptr */
-	FViewModelPtr(FViewModel* InModel)
-		: Model(InModel->AsShared())
-	{}
-	/** Assignment from a raw ptr */
-	FViewModelPtr& operator=(FViewModel* InModel)
+	/** Assignment from a shared ptr */
+	template<typename ModelType>
+	TViewModelPtr& operator=(const TSharedRef<ModelType>& InModel)
 	{
-		Model = InModel->AsShared();
+		static_assert(TPointerIsConvertibleFromTo<ModelType, T>::Value && TPointerIsConvertibleFromTo<ModelType, ViewModelType>::Value,
+			"Construction from shared pointers is only supported for related types. Please use CastViewModel<> to cast to the correct type.");
+
+		Storage.Set(TSharedPtr<ModelType>(InModel), &InModel.Get());
 		return *this;
 	}
 
+
 public:
 
 	/**
@@ -91,7 +201,7 @@ public:
 	 */
 	explicit operator bool() const
 	{
-		return Model.IsValid();
+		return Storage.SharedModel.IsValid();
 	}
 
 	/**
@@ -99,326 +209,161 @@ public:
 	 */
 	void Reset()
 	{
-		Model = nullptr;
+		Storage.Set(nullptr, nullptr);
 	}
 
 	/**
 	 * Member access operator giving access to the wrapped view model
 	 */
-	FViewModel* operator->() const
+	T* operator->() const
 	{
-		checkf(Model, TEXT("Attempting to access a nullptr FViewModelPtr."));
-		return Model.Get();
+		T* Ptr = Storage.Get();
+		checkf(Ptr, TEXT("Attempting to access a nullptr."));
+		return Ptr;
 	}
 
 	/**
 	 * Member access operator giving access to the wrapped view model
 	 */
-	FViewModel* Get() const
+	T* Get() const
 	{
-		return Model.Get();
+		return Storage.Get();
 	}
 
 	/**
 	 * Dereference operator giving access to the wrapped view model
 	 */
-	FViewModel& operator*() const
+	T& operator*() const
 	{
-		checkf(Model, TEXT("Attempting to access a nullptr FViewModelPtr."));
-		return *Model.Get();
+		T* Ptr = Storage.Get();
+		checkf(Ptr, TEXT("Attempting to access a nullptr."));
+		return *Ptr;
 	}
 
 	/**
-	 * Implicit conversion to a regular shared pointer
+	 * Explicitly access the view model as a regular shared pointer
 	 */
-	operator TSharedPtr<FViewModel>() const
+	TSharedPtr<ViewModelType> AsModel() const
 	{
-		return Model;
+		return Storage.SharedModel;
 	}
-
-	/**
-	 * Explicitly access the view model as a regular shared pointer
-	 */
-	TSharedPtr<FViewModel> AsModel() const
+	TSharedRef<T> ToSharedRef() const
 	{
-		return Model;
+		return TSharedPtr<T>(*this).ToSharedRef();
+	}
+	bool IsValid() const
+	{
+		return Storage.Extension != nullptr;
 	}
 
 	/**
 	 * Return a proxy type that is able to implicitly convert this view
 	 * model to any other TViewModel<IMyExtension> or TSharedPtr<IMyExtension>
 	 */
-	SEQUENCERCORE_API FImplicitViewModelCast ImplicitCast() const;
+	TImplicitViewModelCast<ViewModelType> ImplicitCast() const;
+	TImplicitViewModelCastChecked<ViewModelType> ImplicitCastChecked() const;
 
 	/**
 	 * Type hash function for use with the TSet and TMap family of containers
 	 */
-	friend uint32 GetTypeHash(const FViewModelPtr& In)
+	friend uint32 GetTypeHash(const TViewModelPtr<T>& In)
 	{
-		return GetTypeHash(In.Model);
+		return GetTypeHash(In.Storage.SharedModel);
 	}
 
 	/**
 	 * Equality operator comparing 2 view model pointers
 	 */
-	friend bool operator==(const FViewModelPtr& A, const FViewModelPtr& B)
+	friend bool operator==(const TViewModelPtr<T>& A, const TViewModelPtr<T>& B)
 	{
-		return A.Model == B.Model;
+		return A.Storage.SharedModel == B.Storage.SharedModel;
 	}
-	friend bool operator!=(const FViewModelPtr& A, const FViewModelPtr& B)
+	friend bool operator!=(const TViewModelPtr<T>& A, const TViewModelPtr<T>& B)
 	{
-		return A.Model != B.Model;
+		return A.Storage.SharedModel != B.Storage.SharedModel;
 	}
 
-protected:
-
-	TSharedPtr<FViewModel> Model;
-};
-
-
-/**
- * Templated view model pointer that enables conversion and access to an interface implemented by a view model ptr.
- *
- * Since TViewModelPtr is a templated type, it can also be used as a self-documenting mechanism to specify
- *   that a parameter should implement a specific extension interface while only passing a single parameter, for example:
- *   void ExpandRecursively(const TViewModelPtr<IOutlinerExtension>& OutlinerItem);
- *
- * Validity is transitive meaning that an invalid TViewModelPtr<A> is equivalent to a nullptr, and will return an
- * invalid TViewModelPtr<B> even if the originating model implements B but not A.
- */
-template<typename ExtensionType>
-struct TViewModelPtr : FViewModelPtr
-{
-	TViewModelPtr() = default;
-	TViewModelPtr(const TViewModelPtr<ExtensionType>&) = default;
-	TViewModelPtr<ExtensionType>& operator=(const TViewModelPtr<ExtensionType>&) = default;
-	TViewModelPtr(TViewModelPtr<ExtensionType>&&) = default;
-	TViewModelPtr<ExtensionType>& operator=(TViewModelPtr<ExtensionType>&&) = default;
-
-	/** Construction from a regular view model ptr */
-	TViewModelPtr(const FViewModelPtr& InModelPtr)
-		: Extension(nullptr)
-	{
-		Assign(InModelPtr);
-	}
-	/** Assignment from a regular view model ptr */
-	TViewModelPtr& operator=(const FViewModelPtr& InModelPtr)
-	{
-		Assign(InModelPtr);
-		return *this;
-	}
-
-	/** Construction from a regular view model ptr */
-	TViewModelPtr(const TSharedPtr<FViewModel>& InModel)
-		: Extension(nullptr)
-	{
-		Assign(InModel);
-	}
-	/** Assignment from a regular view model ptr */
-	TViewModelPtr& operator=(const TSharedPtr<FViewModel>& InModel)
-	{
-		Assign(InModel);
-		return *this;
-	}
-
-	/** Construction from a regular view model ref */
-	TViewModelPtr(const TSharedRef<FViewModel>& InModel)
-		: Extension(nullptr)
-	{
-		Assign(InModel);
-	}
-	/** Assignment from a regular view model ref */
-	TViewModelPtr& operator=(const TSharedRef<FViewModel>& InModel)
-	{
-		Assign(InModel);
-		return *this;
-	}
-
-	/** Construction from a regular view model ptr */
 	template<typename SharedType>
-	TViewModelPtr(const TSharedPtr<SharedType>& InModel)
+	friend bool operator==(const TViewModelPtr<T>& A, const TSharedPtr<SharedType>& B)
 	{
-		Assign(InModel, InModel.Get());
+		return A.Storage.SharedModel == B;
 	}
-	/** Assignment from a regular view model ptr */
 	template<typename SharedType>
-	TViewModelPtr& operator=(const TSharedPtr<SharedType>& InModel)
+	friend bool operator!=(const TViewModelPtr<T>& A, const TSharedPtr<SharedType>& B)
 	{
-		Assign(InModel, InModel.Get());
-		return *this;
+		return A.Storage.SharedModel != B;
 	}
 
-	/** Construction from a regular view model ref */
 	template<typename SharedType>
-	TViewModelPtr(const TSharedRef<SharedType>& InModel)
+	friend bool operator==(const TViewModelPtr<T>& A, const TSharedRef<SharedType>& B)
 	{
-		Assign(InModel, &InModel.Get());
+		return A.Storage.SharedModel == B;
 	}
-	/** Assignment from a regular view model ptr */
 	template<typename SharedType>
-	TViewModelPtr& operator=(const TSharedRef<SharedType>& InModel)
-	{
-		Assign(InModel, &InModel.Get());
-		return *this;
-	}
-
-	/** Construction from a raw ptr */
-	TViewModelPtr(FViewModel* InModel)
-	{
-		if (InModel)
-		{
-			Assign(InModel->AsShared());
-		}
-	}
-	/** Assignment from a raw ptr */
-	TViewModelPtr<ExtensionType>& operator=(FViewModel* InModel)
+	friend bool operator!=(const TViewModelPtr<T>& A, const TSharedRef<SharedType>& B)
 	{
-		if (InModel)
-		{
-			Assign(InModel->AsShared());
-		}
-		else
-		{
-			Reset();
-		}
-		return *this;
-	}
-
-	/**
-	 * Implicit construction from another view model ptr of a related type (ie, TViewModelPtr<Base>(TViewModelPtr<Derived>) )
-	 * We leave the compiler to report errors on the conversion of In.Extension -> Extension for unrelated types
-	 */
-	template<typename OtherExtensionType>
-	TViewModelPtr(const TViewModelPtr<OtherExtensionType>& In)
-		: FViewModelPtr(In.Model)
-		, Extension(In.Extension)
-	{}
-	/**
-	 * Implicit assignment from another view model ptr of a related type (ie, TViewModelPtr<Base>(TViewModelPtr<Derived>) )
-	 * We leave the compiler to report errors on the conversion of In.Extension -> Extension for unrelated types
-	 */
-	template<typename OtherExtensionType>
-	TViewModelPtr<ExtensionType>& operator=(const TViewModelPtr<OtherExtensionType>& In)
-	{
-		Model = In.Model;
-		Extension = In.Extension;
-		return *this;
-	}
-
-	/**
-	 * Construction from nullptr
-	 */
-	TViewModelPtr(nullptr_t)
-		: FViewModelPtr(nullptr)
-		, Extension(nullptr)
-	{}
-	/**
-	 * Assignment from nullptr
-	 */
-	TViewModelPtr<ExtensionType>& operator=(nullptr_t)
-	{
-		Model = nullptr;
-		Extension = nullptr;
-		return *this;
-	}
-
-	/**
-	 * Conversion to a normal shared pointer of this extension type
-	 */
-	operator TSharedPtr<ExtensionType>() const
-	{
-		return TSharedPtr<ExtensionType>(Model, Extension);
-	}
-
-	/**
-	 * Conversion to a normal weak pointer of this extension type
-	 */
-	operator TWeakPtr<ExtensionType>() const
-	{
-		return TSharedPtr<ExtensionType>(Model, Extension);
-	}
-
-public:
-
-	/**
-	 * Check whether this ptr is valid (non-null)
-	 * @note: While a TViewModelPtr may contain a valid model ptr, it is only
-	 *        considered valid as long as the model implements the templated extension or model type
-	 */
-	explicit operator bool() const
-	{
-		return Extension != nullptr;
-	}
-
-	/**
-	 * Pointer access to the extension pointer. This Ptr must be checked for validity before calling.
-	 */
-	ExtensionType* operator->() const
-	{
-		checkf(Extension, TEXT("Attempting to access a nullptr TViewModelPtr<%s>."), GetGeneratedTypeName<ExtensionType>());
-		return Extension;
+		return A.Storage.SharedModel != B;
 	}
 
-	/**
-	 * Dereference operator giving access to the wrapped extension type
-	 */
-	ExtensionType& operator*() const
+	friend bool operator==(const TViewModelPtr<T>& A, nullptr_t B)
 	{
-		checkf(Extension, TEXT("Attempting to access a nullptr FViewModelPtr."));
-		return *Extension;
+		return A.Storage.SharedModel == nullptr;
 	}
-
-	/**
-	 * Conversion to a normal shared pointer of this extension type
-	 */
-	ExtensionType* Get() const
+	friend bool operator!=(const TViewModelPtr<T>& A, nullptr_t B)
 	{
-		return Extension;
+		return A.Storage.SharedModel != nullptr;
 	}
+protected:
+	template<typename>
+	friend struct TViewModelPtr;
 
-	/**
-	 * Reset this ptr back to its default (null) state
-	 */
-	void Reset()
-	{
-		Model = nullptr;
-		Extension = nullptr;
-	}
+	friend TViewModelConversions<T>;
 
-	friend bool operator==(const TViewModelPtr<ExtensionType>& A, const TViewModelPtr<ExtensionType>& B)
-	{
-		return A.Model == B.Model && A.Extension == B.Extension;
-	}
+	TViewModelPtrStorage<T> Storage;
+};
 
-private:
 
-	void Assign(const TSharedPtr<FViewModel>& InModel, ExtensionType* InExtension = nullptr)
+template<typename OutModelType, typename InModelType>
+TViewModelPtr<OutModelType> CastViewModel(const TSharedRef<InModelType>& InViewModel)
+{
+	return TViewModelPtr<OutModelType>(InViewModel, InViewModel->template CastThis<OutModelType>());
+}
+template<typename OutModelType, typename InModelType>
+TViewModelPtr<OutModelType> CastViewModel(const TSharedPtr<InModelType>& InViewModel)
+{
+	if (!InViewModel)
 	{
-		if (InExtension)
-		{
-			Extension = InExtension;
-			Model = InModel;
-			return;
-		}
-
-		if (InModel)
-		{
-			if (ExtensionType* NewExtension = InModel->CastThis<ExtensionType>())
-			{
-				Extension = NewExtension;
-				Model = InModel;
-				return;
-			}
-		}
-
-		Extension = nullptr;
-		Model = nullptr;
+		return nullptr;
 	}
+	return TViewModelPtr<OutModelType>(InViewModel, InViewModel->template CastThis<OutModelType>());
+}
+template<typename OutModelType, typename InModelType>
+TViewModelPtr<OutModelType> CastViewModel(const TViewModelPtr<InModelType>& InViewModel)
+{
+	return InViewModel.ImplicitCast();
+}
 
-	template<typename> friend struct TViewModelPtr;
-
-	ExtensionType* Extension = nullptr;
-};
+template<typename OutModelType, typename InModelType>
+TViewModelPtr<OutModelType> CastViewModelChecked(const TSharedRef<InModelType>& InViewModel)
+{
+	OutModelType* Extension = InViewModel->template CastThis<OutModelType>();
+	check(Extension);
+	return TViewModelPtr<OutModelType>(InViewModel, Extension);
+}
+template<typename OutModelType, typename InModelType>
+TViewModelPtr<OutModelType> CastViewModelChecked(const TSharedPtr<InModelType>& InViewModel)
+{
+	OutModelType* Extension = InViewModel ? InViewModel->template CastThis<OutModelType>() : nullptr;
+	check(Extension);
+	return TViewModelPtr<OutModelType>(InViewModel, Extension);
+}
+template<typename OutModelType, typename InModelType>
+TViewModelPtr<OutModelType> CastViewModelChecked(const TViewModelPtr<InModelType>& InViewModel)
+{
+	TViewModelPtr<OutModelType> Result = InViewModel.ImplicitCast();
+	check(Result);
+	return Result;
+}
 
 /**
  * Weak pointer utility for view models that also provides the same explicit conversion utility as F/TViewModelPtr
@@ -438,190 +383,112 @@ private:
  *        FWeakViewModelPtr WeakViewModel;
  *    };
  */
-struct FWeakViewModelPtr
+
+template<typename T>
+struct TWeakViewModelConversions
 {
-	FWeakViewModelPtr() = default;
+	using ViewModelType = typename std::conditional_t<std::is_const_v<T>, const FViewModel, FViewModel>;
 
-	/**
-	 * Construction from a shared ref to a view model
-	 * This is templated to ensure that the overloaded constructors remain unambiguous
-	 */
-	template<typename SharedType>
-	FWeakViewModelPtr(const TSharedRef<SharedType>& InModel)
-		: WeakModel(InModel)
-	{}
+	operator TWeakPtr<ViewModelType>() const
+	{
+		return static_cast<const TWeakViewModelPtr<T>*>(this)->WeakModel;
+	}
+};
 
-	/**
-	 * Construction from a shared ptr to a view model
-	 * This is templated to ensure that the overloaded constructors remain unambiguous
-	 */
-	template<typename SharedType>
-	FWeakViewModelPtr(const TSharedPtr<SharedType>& InModel)
-		: WeakModel(InModel)
-	{}
+/**
+ * Templated version of FWeakViewModelPtr that provides an additional self-documenting extension that should
+ * exist on the view model.
+ */
+template<typename T>
+struct TWeakViewModelPtr : TWeakViewModelConversions<T>
+{
+	using ViewModelType = typename std::conditional_t<std::is_const_v<T>, const FViewModel, FViewModel>;
 
-	/**
-	 * Construction from a weak ptr to a view model
-	 * This is templated to ensure that the overloaded constructors remain unambiguous
-	 */
-	template<typename SharedType>
-	FWeakViewModelPtr(const TWeakPtr<SharedType>& InWeakModel)
-		: WeakModel(InWeakModel)
-	{}
+	TWeakViewModelPtr() = default;
+	TWeakViewModelPtr(const TWeakViewModelPtr&) = default;
+	TWeakViewModelPtr& operator=(const TWeakViewModelPtr&) = default;
+	TWeakViewModelPtr(TWeakViewModelPtr&&) = default;
+	TWeakViewModelPtr& operator=(TWeakViewModelPtr&&) = default;
 
-	/**
-	 * Construction from a view model ptr
-	 */
-	FWeakViewModelPtr(const FViewModelPtr& InViewModelPtr)
-		: WeakModel(InViewModelPtr.AsModel())
+	TWeakViewModelPtr(nullptr_t)
 	{}
 
-	/**
-	 * Assignment from a view model ptr
-	 */
-	FWeakViewModelPtr& operator=(const FViewModelPtr& InViewModelPtr)
+	template<typename OtherType>
+	TWeakViewModelPtr(const TWeakPtr<OtherType>& InModel)
 	{
-		WeakModel = InViewModelPtr.AsModel();
-		return *this;
+		static_assert(TPointerIsConvertibleFromTo<OtherType, T>::Value,
+			"Construction from shared pointers is only supported for related types. Please use CastViewModel<> to cast to the correct type.");
+
+		WeakModel = InModel;
 	}
+	template<typename OtherType>
+	TWeakViewModelPtr(const TSharedPtr<OtherType>& InModel)
+	{
+		static_assert(TPointerIsConvertibleFromTo<OtherType, T>::Value,
+			"Construction from shared pointers is only supported for related types. Please use CastViewModel<> to cast to the correct type.");
 
-	/**
-	 * Pin this weak ptr. The result is guaranteed to be kept alive as long as the variable's scope.
-	 */
-	FViewModelPtr Pin() const
+		WeakModel = InModel;
+	}
+	template<typename OtherType>
+	TWeakViewModelPtr(const TSharedRef<OtherType>& InModel)
 	{
-		TSharedPtr<FViewModel> Pinned = WeakModel.Pin();
-		return FViewModelPtr(Pinned);
+		static_assert(TPointerIsConvertibleFromTo<OtherType, T>::Value,
+			"Construction from shared pointers is only supported for related types. Please use CastViewModel<> to cast to the correct type.");
+
+		WeakModel = InModel;
 	}
+	/** Construction from a TViewModelPtr */
+	template<typename OtherType>
+	TWeakViewModelPtr(const TViewModelPtr<OtherType>& InModel)
+	{
+		static_assert(std::is_same_v<T, ViewModelType> || TPointerIsConvertibleFromTo<OtherType, T>::Value,
+			"Construction from shared pointers is only supported for related types. Please use CastViewModel<> to cast to the correct type.");
 
-	/**
-	 * Returns whether the weak ptr is valid.
-	 */
-	bool IsValid() const
+		WeakModel = InModel.AsModel();
+	}
+	/** Assignment from a TViewModelPtr */
+	template<typename OtherType>
+	TWeakViewModelPtr& operator=(const TViewModelPtr<OtherType>& InModel)
 	{
-		return WeakModel.IsValid();
+		static_assert(std::is_same_v<T, ViewModelType> || TPointerIsConvertibleFromTo<OtherType, T>::Value,
+			"Construction from shared pointers is only supported for related types. Please use CastViewModel<> to cast to the correct type.");
+
+		WeakModel = InModel.AsModel();
+		return *this;
 	}
 
-	/**
-	 * Conversion to a normal weak pointer
-	 */
-	operator TWeakPtr<FViewModel>() const
+	/** Pin this weak view model ptr */
+	TViewModelPtr<T> Pin() const
 	{
-		return WeakModel;
+		TSharedPtr<ViewModelType> Model = WeakModel.Pin();
+		return CastViewModel<T>(Model);
 	}
 
-	/**
-	 * Attempt to convert this weak view model ptr to another type implicitly
-	 */
-	SEQUENCERCORE_API FImplicitWeakViewModelCast ImplicitCast() const;
-
 	/**
 	 * Returns a proxy object that will simultaneously Pin() and Cast() this view model
 	 * to a particular view-model or extension type
 	 */
-	SEQUENCERCORE_API FImplicitWeakViewModelPin ImplicitPin() const;
+	TImplicitWeakViewModelPin<ViewModelType> ImplicitPin() const;
 
-	TWeakPtr<FViewModel> AsModel() const
-	{
-		return WeakModel;
-	}
-
-	friend uint32 GetTypeHash(const FWeakViewModelPtr& In)
+	friend uint32 GetTypeHash(const TWeakViewModelPtr<T>& In)
 	{
 		return GetTypeHash(In.WeakModel);
 	}
 
-	friend bool operator==(const FWeakViewModelPtr& A, const FWeakViewModelPtr& B)
+	friend bool operator==(const TWeakViewModelPtr<T>& A, const TWeakViewModelPtr<T>& B)
 	{
 		return A.WeakModel == B.WeakModel;
 	}
 
-	friend bool operator!=(const FWeakViewModelPtr& A, const FWeakViewModelPtr& B)
+	friend bool operator!=(const TWeakViewModelPtr<T>& A, const TWeakViewModelPtr<T>& B)
 	{
 		return A.WeakModel != B.WeakModel;
 	}
 protected:
 
-	TWeakPtr<FViewModel> WeakModel;
-};
+	friend TWeakViewModelConversions<T>;
 
-/**
- * Templated version of FWeakViewModelPtr that provides an additional self-documenting extension that should
- * exist on the view model.
- */
-template<typename ExtensionType>
-struct TWeakViewModelPtr : FWeakViewModelPtr
-{
-	TWeakViewModelPtr() = default;
-	TWeakViewModelPtr(const TWeakViewModelPtr<ExtensionType>&) = default;
-	TWeakViewModelPtr<ExtensionType>& operator=(const TWeakViewModelPtr<ExtensionType>&) = default;
-	TWeakViewModelPtr(TWeakViewModelPtr<ExtensionType>&&) = default;
-	TWeakViewModelPtr<ExtensionType>& operator=(TWeakViewModelPtr<ExtensionType>&&) = default;
-
-	/**
-	 * Construction from a shared ref to a view model
-	 * This is templated to ensure that the overloaded constructors remain unambiguous
-	 */
-	template<typename SharedType>
-	TWeakViewModelPtr(const TSharedRef<SharedType>& InModel)
-		: FWeakViewModelPtr(InModel)
-	{}
-
-	/**
-	 * Construction from a shared ptr to a view model
-	 * This is templated to ensure that the overloaded constructors remain unambiguous
-	 */
-	template<typename SharedType>
-	TWeakViewModelPtr(const TSharedPtr<SharedType>& InModel)
-		: FWeakViewModelPtr(InModel)
-	{}
-
-	/**
-	 * Construction from a weak ptr to a view model
-	 * This is templated to ensure that the overloaded constructors remain unambiguous
-	 */
-	template<typename SharedType>
-	TWeakViewModelPtr(const TWeakPtr<SharedType>& InWeakModel)
-		: FWeakViewModelPtr(InWeakModel)
-	{}
-
-	/** Construction from a FWeakViewModelPtr */
-	TWeakViewModelPtr(const FWeakViewModelPtr& InModel)
-		: FWeakViewModelPtr(InModel)
-	{}
-	/** Assignment from a FWeakViewModelPtr */
-	TWeakViewModelPtr<ExtensionType>& operator=(const FWeakViewModelPtr& InModel)
-	{
-		WeakModel = InModel.AsModel();
-		return *this;
-	}
-
-	/** Construction from a TViewModelPtr */
-	TWeakViewModelPtr(const TViewModelPtr<ExtensionType>& InModel)
-		: FWeakViewModelPtr(InModel.AsModel())
-	{}
-	/** Assignment from a TViewModelPtr */
-	TWeakViewModelPtr<ExtensionType>& operator=(const TViewModelPtr<ExtensionType>& InModel)
-	{
-		WeakModel = InModel.AsModel();
-		return *this;
-	}
-
-	/** Pin this weak view model ptr */
-	TViewModelPtr<ExtensionType> Pin() const
-	{
-		TSharedPtr<FViewModel> Pinned = WeakModel.Pin();
-		return TViewModelPtr<ExtensionType>(Pinned);
-	}
-
-	/**
-	 * Conversion to a normal weak pointer
-	 */
-	operator TWeakPtr<ExtensionType>() const
-	{
-		TViewModelPtr<ExtensionType> Pinned = Pin();
-		return TSharedPtr<ExtensionType>(Pinned);
-	}
+	TWeakPtr<ViewModelType> WeakModel;
 };
 
 
@@ -633,10 +500,11 @@ struct TWeakViewModelPtr : FWeakViewModelPtr
  *
  * @see FViewModelPtr::ImplicitCast
  */
-struct FImplicitViewModelCast
+template<typename ViewModelType, bool bChecked>
+struct TImplicitViewModelCastImpl
 {
 	/** Construction from a regular shared ptr to a view model */
-	explicit FImplicitViewModelCast(const TSharedPtr<FViewModel>& In)
+	explicit TImplicitViewModelCastImpl(const TSharedPtr<ViewModelType>& In)
 		: ViewModel(In)
 	{}
 
@@ -644,10 +512,17 @@ struct FImplicitViewModelCast
 	 * Implicitly convert this view model to any other view-model or extension type
 	 * @return A regular shared pointer to the view-model or extension if it is implemented on the existing model; nullptr otherwise.
 	 */
-	template<typename T>
-	operator TSharedPtr<T>() const
+	template<typename ExtensionType>
+	operator TSharedPtr<ExtensionType>() const
 	{
-		return ViewModel ? ViewModel->CastThisShared<T>() : nullptr;
+		using FinalExtensionType = std::conditional_t<std::is_const_v<ViewModelType>, const ExtensionType, ExtensionType>;
+
+		TViewModelPtr<FinalExtensionType> Result = ViewModel ? ViewModel->template CastThisShared<FinalExtensionType>() : nullptr;
+		if constexpr (bChecked)
+		{
+			checkf(Result, TEXT("Failed to cast view model"));
+		}
+		return Result;
 	}
 
 	/**
@@ -657,55 +532,41 @@ struct FImplicitViewModelCast
 	template<typename ExtensionType>
 	operator TViewModelPtr<ExtensionType>() const
 	{
-		return TViewModelPtr<ExtensionType>(ViewModel);
-	}
-
-private:
-
-	TSharedPtr<FViewModel> ViewModel;
-};
-
-
+		using FinalExtensionType = std::conditional_t<std::is_const_v<ViewModelType>, const ExtensionType, ExtensionType>;
 
-/**
- * Proxy type that enables implicit dynamic casting of a weak view-model to any other view-model
- * or extension type. This is implemented as a proxy type rather than the conversion
- * operators existing on FWeakViewModelPtr so that the implicit cast must be user-directed, while
- * still removing additional template boilerplate at the call-site.
- *
- * @see FWeakViewModelPtr::ImplicitCast
- * @see FWeakViewModelPtr::ImplicitPin
- */
-struct FImplicitWeakViewModelCast
-{
-	/** Construction from a regular weak ptr to a view model */
-	explicit FImplicitWeakViewModelCast(const TWeakPtr<FViewModel>& In)
-		: WeakModel(In)
-	{}
+		TViewModelPtr<FinalExtensionType> Result = ViewModel ? ViewModel->template CastThisShared<ExtensionType>() : nullptr;
+		if constexpr (bChecked)
+		{
+			checkf(Result, TEXT("Failed to cast view model"));
+		}
+		return Result;
+	}
 
 	/**
-	 * Implicit conversion to any TWeakViewModelPtr type
+	 * Implicitly convert this view model to any other view-model or extension type
+	 * @return A TWeakViewModelPtr to the view-model or extension if it is implemented on the existing model; nullptr otherwise.
 	 */
 	template<typename ExtensionType>
 	operator TWeakViewModelPtr<ExtensionType>() const
 	{
-		return TWeakViewModelPtr<ExtensionType>(WeakModel);
+		return TViewModelPtr<ExtensionType>(*this);
 	}
 
 private:
-	TWeakPtr<FViewModel> WeakModel;
-};
 
+	TSharedPtr<ViewModelType> ViewModel;
+};
 
 /**
  * Proxy type that combines a Pin() and Cast() of a weak view model ptr, gracefully handling null.
  *
  * @see FWeakViewModelPtr::ImplicitPin
  */
-struct FImplicitWeakViewModelPin
+template<typename ViewModelType>
+struct TImplicitWeakViewModelPin
 {
 	/** Construction from a regular weak ptr to a view model */
-	explicit FImplicitWeakViewModelPin(const TWeakPtr<FViewModel>& In)
+	explicit TImplicitWeakViewModelPin(const TWeakPtr<ViewModelType>& In)
 		: WeakModel(In)
 	{}
 
@@ -713,11 +574,16 @@ struct FImplicitWeakViewModelPin
 	 * Simultaneously Pin() and Cast() this view model to the templated type
 	 * @return A shared pointer to the view-model or extension type if it is implemented on this model, nullptr otherwise.
 	 */
-	template<typename T>
-	operator TSharedPtr<T>() const
+	template<typename ExtensionType>
+	operator TSharedPtr<ExtensionType>() const
 	{
-		TSharedPtr<FViewModel> Pinned = WeakModel.Pin();
-		return Pinned ? Pinned->CastThisShared<T>() : nullptr;
+		using FinalExtensionType = std::conditional_t<std::is_const_v<ViewModelType>, const ExtensionType, ExtensionType>;
+		TSharedPtr<ViewModelType> Pinned = WeakModel.Pin();
+		if (Pinned)
+		{
+			return TSharedPtr<ExtensionType>(Pinned, Pinned->template CastThis<ExtensionType>());
+		}
+		return nullptr;
 	}
 
 	/**
@@ -727,20 +593,43 @@ struct FImplicitWeakViewModelPin
 	template<typename ExtensionType>
 	operator TViewModelPtr<ExtensionType>() const
 	{
-		TSharedPtr<FViewModel> Pinned = WeakModel.Pin();
+		using FinalExtensionType = std::conditional_t<std::is_const_v<ViewModelType>, const ExtensionType, ExtensionType>;
+
+		TSharedPtr<ViewModelType> Pinned = WeakModel.Pin();
 		if (Pinned)
 		{
-			return TViewModelPtr<ExtensionType>(Pinned);
+			return TViewModelPtr<FinalExtensionType>(Pinned, Pinned->template CastThis<FinalExtensionType>());
 		}
 		return nullptr;
 	}
 
+	template<typename ExtensionType>
+	operator TWeakViewModelPtr<ExtensionType>() const
+	{
+		TViewModelPtr<ExtensionType> Pinned = *this;
+		return Pinned;
+	}
 private:
 
-	TWeakPtr<FViewModel> WeakModel;
+	TWeakPtr<ViewModelType> WeakModel;
 };
 
 
-} // namespace Sequencer
-} // namespace UE
+template<typename T>
+TImplicitViewModelCast<typename TViewModelPtr<T>::ViewModelType> TViewModelPtr<T>::ImplicitCast() const
+{
+	return TImplicitViewModelCast<ViewModelType>{Storage.SharedModel};
+}
+template<typename T>
+TImplicitViewModelCastChecked<typename TViewModelPtr<T>::ViewModelType> TViewModelPtr<T>::ImplicitCastChecked() const
+{
+	return TImplicitViewModelCastChecked<ViewModelType>{Storage.SharedModel};
+}
+
+template<typename T>
+TImplicitWeakViewModelPin<typename TWeakViewModelPtr<T>::ViewModelType> TWeakViewModelPtr<T>::ImplicitPin() const
+{
+	return TImplicitWeakViewModelPin<ViewModelType>{WeakModel};
+}
 
+} // namespace UE::Sequencer
diff --git a/Engine/Source/Editor/SequencerCore/Public/MVVM/ViewModelTypeID.h b/Engine/Source/Editor/SequencerCore/Public/MVVM/ViewModelTypeID.h
--- a/Engine/Source/Editor/SequencerCore/Public/MVVM/ViewModelTypeID.h
+++ b/Engine/Source/Editor/SequencerCore/Public/MVVM/ViewModelTypeID.h
@@ -3,15 +3,31 @@
 #pragma once
 
 #include "CoreTypes.h"
-
-#define UE_SEQUENCER_DECLARE_VIEW_MODEL_TYPE_ID(Type) static ::UE::Sequencer::TAutoRegisterViewModelTypeID<Type> ID;
-#define UE_SEQUENCER_DEFINE_VIEW_MODEL_TYPE_ID(Type) ::UE::Sequencer::TAutoRegisterViewModelTypeID<Type> Type::ID;
+#include "MVVM/CastableTypeTable.h"
+
+#define UE_SEQUENCER_DECLARE_VIEW_MODEL_TYPE_ID(Type)                                               \
+	static ::UE::Sequencer::TAutoRegisterViewModelTypeID<Type> ID;                                  \
+	static void RegisterTypeID();
+
+#define UE_SEQUENCER_DECLARE_VIEW_MODEL_TYPE_ID_API(MODULE_API, Type)                               \
+	MODULE_API static ::UE::Sequencer::TAutoRegisterViewModelTypeID<Type> ID;                       \
+	MODULE_API static void RegisterTypeID();
+
+#define UE_SEQUENCER_DEFINE_VIEW_MODEL_TYPE_ID(Type)                                                \
+	::UE::Sequencer::TAutoRegisterViewModelTypeID<Type> Type::ID;                                   \
+	void Type::RegisterTypeID()                                                                     \
+	{                                                                                               \
+		Type::ID.ID        = FViewModelTypeID::RegisterNewID();                                     \
+		Type::ID.TypeTable = FCastableTypeTable::MakeTypeTable<Type>((Type*)0, Type::ID.ID, #Type); \
+	}
 
 namespace UE
 {
 namespace Sequencer
 {
 
+struct FCastableTypeTable;
+
 struct FViewModelTypeID
 {
 	friend uint32 GetTypeHash(FViewModelTypeID In)
@@ -29,49 +45,68 @@ struct FViewModelTypeID
 		return A.ID == B.ID;
 	}
 
-protected:
-
-	SEQUENCERCORE_API static FViewModelTypeID RegisterNew();
-	SEQUENCERCORE_API static FViewModelTypeID Invalid();
+	uint32 GetTypeID() const
+	{
+		return ID;
+	}
 
-private:
+	SEQUENCERCORE_API static uint32 RegisterNewID();
 
-	FViewModelTypeID(uint32 InID)
-		: ID(InID)
+	FViewModelTypeID(FCastableTypeTable* InTypeTable, uint32 InID)
+		: TypeTable(InTypeTable)
+		, ID(InID)
 	{}
 
-	FViewModelTypeID()
-		: ID(~0u)
-	{}
+private:
 
+	FCastableTypeTable* TypeTable;
 	uint32 ID;
 };
 
 template<typename T>
-struct TViewModelTypeID : FViewModelTypeID
+struct TViewModelTypeID
 {
-	static TViewModelTypeID<T> RegisterNew()
+	operator FViewModelTypeID() const
+	{
+		Register();
+		return FViewModelTypeID{ TypeTable, ID };
+	}
+
+	uint32 GetTypeID() const
+	{
+		Register();
+		return ID;
+	}
+
+	FCastableTypeTable* GetTypeTable() const
 	{
-		return TViewModelTypeID<T>(FViewModelTypeID::RegisterNew());
+		Register();
+		return TypeTable;
 	}
 
-	static TViewModelTypeID<T> Invalid()
+	void Register() const
 	{
-		return TViewModelTypeID<T>(FViewModelTypeID::Invalid());
+		if (!IsRegistered())
+		{
+			T::RegisterTypeID();
+		}
 	}
 
 protected:
-	explicit TViewModelTypeID(FViewModelTypeID In)
-		: FViewModelTypeID(In)
-	{}
+	friend T;
+
+	bool IsRegistered() const
+	{
+		return ID != ~0u;
+	}
+
+	FCastableTypeTable* TypeTable = nullptr;
+	uint32 ID = ~0u;
 };
 
 template<typename T>
 struct TAutoRegisterViewModelTypeID : TViewModelTypeID<T>
 {
-	TAutoRegisterViewModelTypeID()
-		: TViewModelTypeID<T>(TViewModelTypeID<T>::RegisterNew())
-	{}
 };
 
 } // namespace Sequencer
diff --git a/Engine/Source/Editor/SequencerCore/Public/MVVM/ViewModels/OutlinerViewModel.h b/Engine/Source/Editor/SequencerCore/Public/MVVM/ViewModels/OutlinerViewModel.h
--- a/Engine/Source/Editor/SequencerCore/Public/MVVM/ViewModels/OutlinerViewModel.h
+++ b/Engine/Source/Editor/SequencerCore/Public/MVVM/ViewModels/OutlinerViewModel.h
@@ -11,7 +11,6 @@
 #include "MVVM/ViewModelPtr.h"
 #include "MVVM/ViewModelTypeID.h"
 #include "MVVM/ViewModels/ViewModel.h"
-#include "MVVM/ViewModels/ViewModel.h"
 #include "Templates/SharedPointer.h"
 
 class FDragDropOperation;
diff --git a/Engine/Source/Editor/SequencerCore/Public/MVVM/ViewModels/ViewModel.h b/Engine/Source/Editor/SequencerCore/Public/MVVM/ViewModels/ViewModel.h
--- a/Engine/Source/Editor/SequencerCore/Public/MVVM/ViewModels/ViewModel.h
+++ b/Engine/Source/Editor/SequencerCore/Public/MVVM/ViewModels/ViewModel.h
@@ -24,7 +24,6 @@ struct FParentFirstChildIterator;
 struct FParentModelIterator;
 struct FViewModelIterationState;
 struct FViewModelListIterator;
-struct FViewModelPtr;
 
 /**
  * Base class for a sequencer data model. This might wrap an underlying UObject, or be a purely
@@ -61,7 +60,10 @@ public:
 		{
 			return *StaticImpl;
 		}
-		return FDynamicExtensionContainer::AddDynamicExtension<T>(AsShared(), Forward<InArgTypes>(Args)...);
+
+		// Assign the dynamic type member so that we can cast directly to these types
+		this->DynamicTypes = this;
+		return FDynamicExtensionContainer::AddDynamicExtension<T>(AsShared(), Forward<InArgTypes>(Args)...);;
 	}
 
 	/** Adds a dynamic extension to this data model */
@@ -73,16 +75,16 @@ public:
 
 	/** Casts this data model to an extension, or to a child class implementation */
 	template<typename T>
-	TSharedPtr<T> CastThisShared()
+	TViewModelPtr<T> CastThisShared()
 	{
-		return TSharedPtr<T>(AsShared(), CastThis<T>());
+		return TViewModelPtr<T>(AsShared(), CastThis<T>());
 	}
 
 	/** Casts this data model to an extension, or to a child class implementation */
 	template<typename T>
-	TSharedPtr<const T> CastThisShared() const
+	TViewModelPtr<const T> CastThisShared() const
 	{
-		return TSharedPtr<const T>(AsShared(), CastThis<T>());
+		return TViewModelPtr<const T>(AsShared(), CastThis<const T>());
 	}
 
 	/** Casts this data model to an extension, or to a child class implementation */
@@ -116,28 +118,18 @@ public:
 	}
 
 	/** Gets the parent data model */
-	TSharedPtr<FViewModel> GetParent() const
-	{
-		return WeakParent.Pin();
-	}
+	FViewModelPtr GetParent() const;
 
 	/** Gets the parent data model */
 	template<typename T>
 	TViewModelPtr<T> CastParent() const
 	{
-		return TViewModelPtr<T>(WeakParent.Pin());
+		TSharedPtr<FViewModel> Parent = WeakParent.Pin();
+		return Parent ? TViewModelPtr<T>(Parent, Parent->CastThis<T>()) : TViewModelPtr<T>();
 	}
 
 	/** Returns the root model of the hierarchy, i.e. the parent model that has no parent */
-	TSharedPtr<FViewModel> GetRoot() const
-	{
-		TSharedPtr<FViewModel> CurrentItem(SharedThis(const_cast<FViewModel*>(this)));
-		while (CurrentItem->GetParent())
-		{
-			CurrentItem = CurrentItem->GetParent();
-		}
-		return CurrentItem;
-	}
+	FViewModelPtr GetRoot() const;
 
 public:
 
diff --git a/Engine/Source/Editor/SequencerCore/Public/MVVM/ViewModels/ViewModelHierarchy.h b/Engine/Source/Editor/SequencerCore/Public/MVVM/ViewModels/ViewModelHierarchy.h
--- a/Engine/Source/Editor/SequencerCore/Public/MVVM/ViewModels/ViewModelHierarchy.h
+++ b/Engine/Source/Editor/SequencerCore/Public/MVVM/ViewModels/ViewModelHierarchy.h
@@ -18,9 +18,7 @@ class FSharedViewModelData;
 class FViewModel;
 struct FViewModelChildren;
 struct FViewModelListHead;
-struct FViewModelPtr;
 struct FViewModelSubListIterator;
-template <typename ExtensionType> struct TViewModelPtr;
 
 #if UE_BUILD_DEBUG || UE_BUILD_DEVELOPMENT
 	#define UE_SEQUENCER_DETECT_LINK_LIST_CYCLES 1
diff --git a/Engine/Source/Editor/SequencerCore/Public/MVVM/ViewModels/ViewModelIterators.h b/Engine/Source/Editor/SequencerCore/Public/MVVM/ViewModels/ViewModelIterators.h
--- a/Engine/Source/Editor/SequencerCore/Public/MVVM/ViewModels/ViewModelIterators.h
+++ b/Engine/Source/Editor/SequencerCore/Public/MVVM/ViewModels/ViewModelIterators.h
@@ -554,7 +554,7 @@ protected:
 	{
 		while (*this)
 		{
-			TypedValue = TViewModelPtr<T>(IteratorType::GetCurrentItem().AsModel());
+			TypedValue = IteratorType::GetCurrentItem()->template CastThisShared<T>();
 			if (TypedValue)
 			{
 				break;
diff --git a/Engine/Source/Editor/SequencerCore/Public/SequencerCoreFwd.h b/Engine/Source/Editor/SequencerCore/Public/SequencerCoreFwd.h
--- a/Engine/Source/Editor/SequencerCore/Public/SequencerCoreFwd.h
+++ b/Engine/Source/Editor/SequencerCore/Public/SequencerCoreFwd.h
@@ -43,19 +43,27 @@ using TViewModelSubListIterator = TTypedIterator<T, FViewModelSubListIterator>;
 /**
  * Shared/weak extension pointer forward declarations - include MVVM/ViewModelPtr.h for complete declarations
  */
-struct FImplicitViewModelCast;
-struct FImplicitWeakViewModelPin;
-struct FImplicitWeakViewModelCast;
+template<typename, bool>
+struct TImplicitViewModelCastImpl;
+template<typename>
+struct TImplicitWeakViewModelPin;
 
-struct FViewModelPtr;
-struct FWeakViewModelPtr;
+template<typename ViewModelType>
+using TImplicitViewModelCast = TImplicitViewModelCastImpl<ViewModelType, false>;
 
-template<typename ExtensionType>
+template<typename ViewModelType>
+using TImplicitViewModelCastChecked = TImplicitViewModelCastImpl<ViewModelType, true>;
+
+template<typename>
 struct TViewModelPtr;
 
-template<typename ExtensionType>
+template<typename>
 struct TWeakViewModelPtr;
 
+using FViewModelPtr = TViewModelPtr<FViewModel>;
+using FConstViewModelPtr = TViewModelPtr<const FViewModel>;
+using FWeakViewModelPtr = TWeakViewModelPtr<FViewModel>;
 
 } // namespace Sequencer
 } // namespace UE
+
