Sequencer: New selection API

diff --git a/Engine/Source/Editor/SequencerCore/Private/MVVM/Views/SOutlinerView.cpp b/Engine/Source/Editor/SequencerCore/Private/MVVM/Views/SOutlinerView.cpp
--- a/Engine/Source/Editor/SequencerCore/Private/MVVM/Views/SOutlinerView.cpp
+++ b/Engine/Source/Editor/SequencerCore/Private/MVVM/Views/SOutlinerView.cpp
@@ -27,7 +27,8 @@
 #include "MVVM/ViewModels/OutlinerViewModelDragDropOp.h"
 #include "MVVM/ViewModels/ViewModel.h"
 #include "MVVM/ViewModels/ViewModelIterators.h"
-#include "MVVM/Views/IOutlinerSelectionHandler.h"
+#include "MVVM/Selection/SequencerCoreSelection.h"
+#include "MVVM/Selection/SequencerOutlinerSelection.h"
 #include "MVVM/Views/STrackAreaView.h"
 #include "MVVM/Views/STrackLane.h"
 #include "MVVM/Views/TreeViewTraits.h"
@@ -40,6 +41,7 @@
 #include "Rendering/SlateLayoutTransform.h"
 #include "ScopedTransaction.h"
 #include "SequencerCoreFwd.h"
+#include "Framework/Application/SlateApplication.h"
 #include "Styling/SlateBrush.h"
 #include "Styling/WidgetStyle.h"
 #include "Templates/Tuple.h"
@@ -254,9 +256,17 @@ void SOutlinerView::Construct(const FArguments& InArgs, TWeakPtr<FOutlinerViewMo
 	bShowPinnedNodes = false;
 	bSelectionChangesPending = false;
 	bRefreshPhysicalGeometry = false;
-	bForwardToSelectionHandler = true;
 
-	SelectionHandlerAttribute = InArgs._SelectionHandler;
+	Selection = InArgs._Selection;
+
+	if (Selection)
+	{
+		TSharedPtr<FOutlinerSelection> OutlinerSelection = Selection->GetOutlinerSelection();
+		if (OutlinerSelection)
+		{
+			OutlinerSelection->OnChanged.AddSP(this, &SOutlinerView::UpdateViewSelectionFromModel);
+		}
+	}
 
 	HeaderRow = SNew(SHeaderRow).Visibility(EVisibility::Collapsed);
 
@@ -282,8 +292,13 @@ void SOutlinerView::Construct(const FArguments& InArgs, TWeakPtr<FOutlinerViewMo
 		.HighlightParentNodesForSelection(true)
 		.AllowInvisibleItemSelection(true)  //without this we deselect everything when we filter or we collapse, etc..
 	);
+
+	UpdateViewSelectionFromModel();
 }
 
+SOutlinerView::SOutlinerView()
+{}
+
 SOutlinerView::~SOutlinerView()
 {
 	TrackArea.Reset();
@@ -296,6 +311,11 @@ void SOutlinerView::Tick(const FGeometry& AllottedGeometry, const double InCurre
 
 	STreeView::Tick(AllottedGeometry, InCurrentTime, InDeltaTime);
 
+	if (!FSlateApplication::Get().AnyMenusVisible())
+	{
+		DelayedEventSuppressor = nullptr;
+	}
+
 	if (bWasPendingRefresh || bRefreshPhysicalGeometry)
 	{
 		UpdatePhysicalGeometry(true);
@@ -327,26 +347,26 @@ FReply SOutlinerView::OnDragRow(const FGeometry&, const FPointerEvent&, TSharedR
 		return FReply::Unhandled();
 	}
 
-	TArray<TWeakViewModelPtr<IOutlinerExtension>> Selection = GetSelectedItems();
-	if (Selection.Num() > 0)
+	TArray<TWeakViewModelPtr<IOutlinerExtension>> WeakSelectedItems = GetSelectedItems();
+	if (WeakSelectedItems.Num() > 0)
 	{
-		for (int32 Index = Selection.Num()-1; Index >= 0; --Index)
+		for (int32 Index = WeakSelectedItems.Num()-1; Index >= 0; --Index)
 		{
-			TSharedPtr<IDraggableOutlinerExtension> Draggable = Selection[Index].ImplicitPin();
+			TSharedPtr<IDraggableOutlinerExtension> Draggable = WeakSelectedItems[Index].ImplicitPin();
 			if (!Draggable || !Draggable->CanDrag())
 			{
 				// Order is not important so we can opt for performance with RemoveAtSwap
-				Selection.RemoveAtSwap(Index, 1, false);
+				WeakSelectedItems.RemoveAtSwap(Index, 1, false);
 			}
 		}
 
 		// If there were no nodes selected, don't start a drag drop operation.
-		if (Selection.Num() == 0)
+		if (WeakSelectedItems.Num() == 0)
 		{
 			return FReply::Unhandled();
 		}
 
-		TSharedRef<FDragDropOperation> DragDropOp = Outliner->InitiateDrag( MoveTemp(Selection) );
+		TSharedRef<FDragDropOperation> DragDropOp = Outliner->InitiateDrag( MoveTemp(WeakSelectedItems) );
 		return FReply::Handled().BeginDragDrop( DragDropOp );
 	}
 
@@ -356,7 +376,7 @@ FReply SOutlinerView::OnDragRow(const FGeometry&, const FPointerEvent&, TSharedR
 void SOutlinerView::ReportChildRowGeometry(const TViewModelPtr<IOutlinerExtension>& InNode, const FGeometry& InGeometry)
 {
 	FCachedGeometry* PhysicalNode = PhysicalNodes.FindByPredicate(
-			[InNode](FCachedGeometry& Item) { return Item.WeakItem == InNode; });
+			[InNode](FCachedGeometry& Item) { return Item.WeakItem.Pin() == InNode; });
 	if (PhysicalNode)
 	{
 		PhysicalNode->PhysicalHeight = InGeometry.Size.Y;
@@ -375,18 +395,12 @@ void SOutlinerView::GetVisibleItems(TArray<TViewModelPtr<IOutlinerExtension>>& O
 	}
 }
 
-void SOutlinerView::ForceSetSelectedItems(const TSet<TWeakPtr<FViewModel>>& InItems)
+void SOutlinerView::ForceSetSelectedItems(const TSet<TWeakViewModelPtr<IOutlinerExtension>>& InItems)
 {
-	TGuardValue<bool> Guard(bForwardToSelectionHandler, false);
-
 	Private_ClearSelection();
-	for (const TWeakPtr<FViewModel>& Item : InItems)
+	for (const TWeakViewModelPtr<IOutlinerExtension>& Item : InItems)
 	{
-		TWeakViewModelPtr<IOutlinerExtension> WeakItem = CastViewModel<IOutlinerExtension>(Item.Pin());
-		if (WeakItem.Pin())
-		{
-			Private_SetItemSelection(WeakItem, true, false);
-		}
+		Private_SetItemSelection(Item, true, false);
 	}
 
 	Private_SignalSelectionChanged(ESelectInfo::Direct);
@@ -774,17 +788,57 @@ void SOutlinerView::Private_SelectRangeFromCurrentTo( TWeakViewModelPtr<IOutline
 
 void SOutlinerView::Private_SignalSelectionChanged(ESelectInfo::Type SelectInfo)
 {
-	// If bForwardToSelectionHandler is set here that means we are responding to a
-	// selection change event that originated from a user interaction with this outliner view.
-	// Therefore we must go through and re-assign selection states to the selection handler.
-
 	STreeView::Private_SignalSelectionChanged(SelectInfo);
 
-	TSharedPtr<IOutlinerSelectionHandler> SelectionHandler = SelectionHandlerAttribute.Get();
-	if (SelectionHandler && bForwardToSelectionHandler)
+	UpdateModelSelectionFromView();
+}
+
+void SOutlinerView::UpdateViewSelectionFromModel()
+{
+	if (!Selection)
+	{
+		return;
+	}
+
+	TSharedPtr<const FOutlinerSelection> OutlinerSelection = Selection->GetOutlinerSelection();
+	if (OutlinerSelection != nullptr)
+	{
+		Private_ClearSelection();
+
+		for (TViewModelPtr<IOutlinerExtension> SelectedItem : *OutlinerSelection)
+		{
+			Private_SetItemSelection(SelectedItem, true, false);
+		}
+
+		Private_SignalSelectionChanged(ESelectInfo::Direct);
+	}
+}
+
+void SOutlinerView::UpdateModelSelectionFromView()
+{
+	if (!Selection || Selection->IsTriggeringSelectionChangedEvents())
+	{
+		return;
+	}
+
+	TSharedPtr<FOutlinerSelection> OutlinerSelection = Selection->GetOutlinerSelection();
+	if (OutlinerSelection == nullptr)
+	{
+		return;
+	}
+
+	DelayedEventSuppressor = Selection->SuppressEventsLongRunning();
+
+	OutlinerSelection->Empty();
+	for (TWeakViewModelPtr<IOutlinerExtension> SelectedItem : GetSelectedItems())
+	{
+		OutlinerSelection->Select(SelectedItem);
+	}
+
+	if (bRightMouseButtonDown == false)
 	{
-		TGuardValue<bool> Guard(bForwardToSelectionHandler, false);
-		SelectionHandler->SelectOutlinerItems(GetSelectedItems(), bRightMouseButtonDown);
+		// If we're not going to open a context menu - trigger events now
+		DelayedEventSuppressor = nullptr;
 	}
 }
 
@@ -984,10 +1038,10 @@ void SOutlinerView::UpdatePhysicalGeometry(bool bIsRefresh)
 	}
 }
 
-bool ShouldExpand(const TArrayView<const TWeakViewModelPtr<IOutlinerExtension>>& Selection, ETreeRecursion Recursion, TSharedPtr<FOutlinerViewModel> Outliner)
+bool ShouldExpand(const TArrayView<const TWeakViewModelPtr<IOutlinerExtension>>& WeakSelectedItems, ETreeRecursion Recursion, TSharedPtr<FOutlinerViewModel> Outliner)
 {
 	bool bAllExpanded = true;
-	for (TWeakViewModelPtr<IOutlinerExtension> WeakItem : Selection)
+	for (TWeakViewModelPtr<IOutlinerExtension> WeakItem : WeakSelectedItems)
 	{
 		// @todo_sequencer_mvvm: Do we need to check selection state here if it's already in a supposed selection set??
 		TViewModelPtr<IOutlinerExtension> SelectedItem = WeakItem.Pin();
@@ -1018,13 +1072,13 @@ void SOutlinerView::ToggleExpandCollapseNodes(ETreeRecursion Recursion, bool bEx
 		return;
 	}
 
-	TArray<TWeakViewModelPtr<IOutlinerExtension>> Selection = GetSelectedItems();
+	TArray<TWeakViewModelPtr<IOutlinerExtension>> WeakSelectedItems = GetSelectedItems();
 
-	if (Selection.Num() > 0 && !bExpandAll && !bCollapseAll)
+	if (WeakSelectedItems.Num() > 0 && !bExpandAll && !bCollapseAll)
 	{
-		const bool bExpand = ShouldExpand(Selection, Recursion, Outliner);
+		const bool bExpand = ShouldExpand(WeakSelectedItems, Recursion, Outliner);
 		
-		for (TWeakViewModelPtr<IOutlinerExtension> WeakItem : Selection)
+		for (TWeakViewModelPtr<IOutlinerExtension> WeakItem : WeakSelectedItems)
 		{
 			TViewModelPtr<IOutlinerExtension> Item = WeakItem.Pin();
 			if (Item)
diff --git a/Engine/Source/Editor/SequencerCore/Private/SequencerCoreModule.cpp b/Engine/Source/Editor/SequencerCore/Private/SequencerCoreModule.cpp
--- a/Engine/Source/Editor/SequencerCore/Private/SequencerCoreModule.cpp
+++ b/Engine/Source/Editor/SequencerCore/Private/SequencerCoreModule.cpp
@@ -2,6 +2,7 @@
 
 #include "SequencerCoreModule.h"
 #include "Modules/ModuleManager.h"
+#include "SequencerCoreLog.h"
 
 namespace UE
 {
@@ -23,4 +24,5 @@ class FSequencerCoreModule
 } // namespace UE
 
 
+DEFINE_LOG_CATEGORY(LogSequencerCore);
 IMPLEMENT_MODULE(UE::Sequencer::FSequencerCoreModule, SequencerCore);
diff --git a/Engine/Source/Editor/SequencerCore/Public/MVVM/ViewModels/EditorViewModel.h b/Engine/Source/Editor/SequencerCore/Public/MVVM/ViewModels/EditorViewModel.h
--- a/Engine/Source/Editor/SequencerCore/Public/MVVM/ViewModels/EditorViewModel.h
+++ b/Engine/Source/Editor/SequencerCore/Public/MVVM/ViewModels/EditorViewModel.h
@@ -11,16 +11,18 @@
 #include "Misc/FrameRate.h"
 #include "Templates/SharedPointer.h"
 #include "TimeToPixel.h"
+#include "UObject/StrongObjectPtr.h"
 
-namespace UE::Sequencer { template <typename T> struct TAutoRegisterViewModelTypeID; }
+class USequencerScriptingLayer;
 
-namespace UE
-{
-namespace Sequencer
+namespace UE::Sequencer
 {
 
 class FOutlinerViewModel;
 class FTrackAreaViewModel;
+class FSequencerCoreSelection;
+
+template <typename T> struct TAutoRegisterViewModelTypeID;
 
 /**
  * This represents to root view-model for a sequencer-like editor.
@@ -63,6 +65,12 @@ public:
 	/** Gets the track area view-model. */
 	TSharedPtr<FTrackAreaViewModel> GetTrackArea() const;
 
+	/** Gets the selection class */
+	TSharedPtr<FSequencerCoreSelection> GetSelection() const;
+
+	/** Gets the scripting layer */
+	USequencerScriptingLayer* GetScriptingLayer() const;
+
 	/** Returns whether this editor is currently read-only */
 	virtual bool IsReadOnly() const;
 
@@ -76,6 +84,10 @@ protected:
 	virtual TSharedPtr<FOutlinerViewModel> CreateOutlinerImpl();
 	/** Creates the track-area for this editor. */
 	virtual TSharedPtr<FTrackAreaViewModel> CreateTrackAreaImpl();
+	/** Creates the selection class for this editor. */
+	virtual TSharedPtr<FSequencerCoreSelection> CreateSelectionImpl();
+	/** Creates the scripting layer class for this editor. */
+	virtual USequencerScriptingLayer* CreateScriptingLayerImpl();
 	/** Creates any other panels for this editor. */
 	virtual void InitializeEditorImpl() {}
 
@@ -90,11 +102,14 @@ private:
 	TSharedPtr<FOutlinerViewModel>  Outliner;
 	/** Cached pointer to the track area panel view-model */
 	TSharedPtr<FTrackAreaViewModel> TrackArea;
+	/** Cached pointer to the selection class */
+	TSharedPtr<FSequencerCoreSelection> Selection;
+	/** Cached pointer to the scripting layer class */
+	TStrongObjectPtr<USequencerScriptingLayer> ScriptingLayer;
 
 	/** The root view-model for the data being edited */
 	TSharedPtr<FViewModel>  RootDataModel;
 };
 
-} // namespace Sequencer
-} // namespace UE
+} // namespace UE::Sequencer
 
diff --git a/Engine/Source/Editor/SequencerCore/Public/MVVM/Views/SOutlinerView.h b/Engine/Source/Editor/SequencerCore/Public/MVVM/Views/SOutlinerView.h
--- a/Engine/Source/Editor/SequencerCore/Public/MVVM/Views/SOutlinerView.h
+++ b/Engine/Source/Editor/SequencerCore/Public/MVVM/Views/SOutlinerView.h
@@ -47,8 +47,10 @@ namespace UE
 namespace Sequencer
 {
 
+struct FSelectionEventSuppressor;
+
 class FViewModel;
-class IOutlinerSelectionHandler;
+class FSequencerCoreSelection;
 class SOutlinerViewRow;
 class STrackAreaView;
 class STrackLane;
@@ -81,7 +83,7 @@ public:
 
 	SLATE_BEGIN_ARGS(SOutlinerView){}
 
-		SLATE_ATTRIBUTE( TSharedPtr<IOutlinerSelectionHandler>, SelectionHandler )
+		SLATE_ARGUMENT( TSharedPtr<FSequencerCoreSelection>, Selection )
 		/** Externally supplied scroll bar */
 		SLATE_ARGUMENT( TSharedPtr<SScrollBar>, ExternalScrollbar )
 
@@ -89,6 +91,7 @@ public:
 
 	static const FName TrackNameColumn;
 
+	SOutlinerView();
 	~SOutlinerView();
 
 	/** Construct this widget */
@@ -105,7 +108,7 @@ public:
 
 	void GetVisibleItems(TArray<TViewModelPtr<IOutlinerExtension>>& OutItems) const;
 
-	void ForceSetSelectedItems(const TSet<TWeakPtr<FViewModel>>& InItems);
+	void ForceSetSelectedItems(const TSet<TWeakViewModelPtr<IOutlinerExtension>>& InItems);
 
 public:
 
@@ -176,6 +179,9 @@ private:
 	// Private, unimplemented overloaded name for SetItemSelection to prevent external calls - use ForceSetItemSelection instead
 	void ClearSelection();
 
+	void UpdateViewSelectionFromModel();
+	void UpdateModelSelectionFromView();
+
 	void HandleTableViewScrolled(double InScrollOffset);
 	void UpdatePhysicalGeometry(bool bIsRefresh);
 
@@ -246,7 +252,8 @@ private:
 	/** Column definitions for each of the columns in the tree view */
 	TMap<FName, FOutlinerViewColumn> Columns;
 
-	TAttribute<TSharedPtr<IOutlinerSelectionHandler>> SelectionHandlerAttribute;
+	TSharedPtr<FSequencerCoreSelection> Selection;
+	TUniquePtr<FSelectionEventSuppressor> DelayedEventSuppressor;
 
 	/** Strong pointer to the track area so we can generate track lanes as we need them */
 	TSharedPtr<STrackAreaView> TrackArea;
@@ -259,8 +266,6 @@ private:
 
 	float VirtualTop;
 
-	bool bForwardToSelectionHandler;
-
 	/** When true, the tree selection is being updated from a change in the sequencer selection. */
 	bool bUpdatingTreeSelection;
 
diff --git a/Engine/Source/Runtime/MovieSceneTracks/Public/Sections/MovieSceneActorReferenceSection.h b/Engine/Source/Runtime/MovieSceneTracks/Public/Sections/MovieSceneActorReferenceSection.h
--- a/Engine/Source/Runtime/MovieSceneTracks/Public/Sections/MovieSceneActorReferenceSection.h
+++ b/Engine/Source/Runtime/MovieSceneTracks/Public/Sections/MovieSceneActorReferenceSection.h
@@ -139,7 +139,7 @@ public:
 	 *
 	 * @return (Optional) The channel's default value
 	 */
-	FORCEINLINE FMovieSceneActorReferenceKey GetDefault() const
+	FORCEINLINE const FMovieSceneActorReferenceKey& GetDefault() const
 	{
 		return DefaultValue;
 	}
