Sequencer: don't assert on zero-width sections

diff --git a/Engine/Source/Editor/SequencerCore/Public/TimeToPixel.h b/Engine/Source/Editor/SequencerCore/Public/TimeToPixel.h
--- a/Engine/Source/Editor/SequencerCore/Public/TimeToPixel.h
+++ b/Engine/Source/Editor/SequencerCore/Public/TimeToPixel.h
@@ -23,7 +23,7 @@ public:
 		: ViewRangeStartSeconds( InLocalViewRange.GetLowerBoundValue() )
 		, TickResolution( InTickResolution )
 	{
-		ensureMsgf(WidthPx > 0, TEXT("Got zero or negative pixel width!"));
+		ensureMsgf(WidthPx >= 0, TEXT("Got negative pixel width!"));
 		const double VisibleWidth = InLocalViewRange.Size<double>();
 
 		const float MaxPixelsPerSecond = 1000.f;
