Sequencer: Fixed a crash when duplicating keys

The tempoarary structure used by duplication has no owner, so doesn't need to report changes

diff --git a/Engine/Source/Editor/SequencerCore/Private/MVVM/Selection/SequencerCoreSelectionTypes.cpp b/Engine/Source/Editor/SequencerCore/Private/MVVM/Selection/SequencerCoreSelectionTypes.cpp
--- a/Engine/Source/Editor/SequencerCore/Private/MVVM/Selection/SequencerCoreSelectionTypes.cpp
+++ b/Engine/Source/Editor/SequencerCore/Private/MVVM/Selection/SequencerCoreSelectionTypes.cpp
@@ -8,17 +8,20 @@ namespace UE::Sequencer
 
 void FSelectionBase::Empty()
 {
-	FSelectionEventSuppressor SuppressEvents = Owner->SuppressEvents();
+	FSelectionEventSuppressor SuppressEvents(Owner);
 	EmptyImpl();
 }
 
 void FSelectionBase::ReportChanges(bool bInSelectionChanged)
 {
-	if (bInSelectionChanged && !bSelectionChanged)
+	if (Owner)
 	{
-		Owner->PreSelectionSetChangeEvent(this);
+		if (bInSelectionChanged && !bSelectionChanged)
+		{
+			Owner->PreSelectionSetChangeEvent(this);
+		}
+		bSelectionChanged |= bInSelectionChanged;
 	}
-	bSelectionChanged |= bInSelectionChanged;
 }
 
 } // namespace UE::Sequencer
\ No newline at end of file
