Sequencer: Outliner column registration and column interface declaration

diff --git a/Engine/Source/Editor/Sequencer/Public/ISequencerModule.h b/Engine/Source/Editor/Sequencer/Public/ISequencerModule.h
--- a/Engine/Source/Editor/Sequencer/Public/ISequencerModule.h
+++ b/Engine/Source/Editor/Sequencer/Public/ISequencerModule.h
@@ -252,7 +252,7 @@ public:
 	 *
 	 * @param InHandle Handle to the delegate to unregister
 	 */
-	virtual void UnregisterOutlinerColumn(FDelegateHandle Inhandle) = 0;
+	virtual void UnregisterOutlinerColumn(FDelegateHandle InHandle) = 0;
 
 
 	/** 
diff --git a/Engine/Source/Editor/SequencerCore/Public/ISequencerOutlinerColumn.h b/Engine/Source/Editor/SequencerCore/Public/ISequencerOutlinerColumn.h
--- a/Engine/Source/Editor/SequencerCore/Public/ISequencerOutlinerColumn.h
+++ b/Engine/Source/Editor/SequencerCore/Public/ISequencerOutlinerColumn.h
@@ -9,14 +9,14 @@ class UMovieSceneSequence;
 
 namespace UE
 {
-	namespace Sequencer
-	{
-		class IOutlinerExtension;
-	}
+namespace Sequencer
+{
+class IOutlinerExtension;
+}
 }
 
 /**
- * Interface for sequencer track editors.
+ * Interface for sequencer outliner columns.
  */
 class ISequencerOutlinerColumn
 {
