Sequencer: Track Gutter UI Framework, Outliner Column Creation and Settings

diff --git a/Engine/Source/Editor/Sequencer/Private/SSequencer.cpp b/Engine/Source/Editor/Sequencer/Private/SSequencer.cpp
--- a/Engine/Source/Editor/Sequencer/Private/SSequencer.cpp
+++ b/Engine/Source/Editor/Sequencer/Private/SSequencer.cpp
@@ -308,6 +308,11 @@ void SSequencer::Construct(const FArguments& InArgs, TSharedRef<FSequencer> InSe
 
 	TreeView->AddPinnedTreeView(PinnedTreeView);
 
+	if (GetSequencerSettings())
+	{
+		InitializeOutlinerColumns();
+	}
+
 	SequencerViewModel->GetTrackArea()->CastThisChecked<FSequencerTrackAreaViewModel>()->InitializeDefaultEditTools(*TrackArea);
 	SequencerViewModel->GetPinnedTrackArea()->CastThisChecked<FSequencerTrackAreaViewModel>()->InitializeDefaultEditTools(*PinnedTrackArea);
 
@@ -528,7 +533,6 @@ void SSequencer::Construct(const FArguments& InArgs, TSharedRef<FSequencer> InSe
 						.Clipping(EWidgetClipping::ClipToBounds)
 						[
 							SNew(SHorizontalBox)
-
 							+ SHorizontalBox::Slot()
 							.AutoWidth()
 							.VAlign(VAlign_Center)
@@ -1074,6 +1078,72 @@ void SSequencer::InitializeTrackFilters()
 	AllTrackFilters.Sort([](const TSharedRef<FSequencerTrackFilter>& LHS, const TSharedRef<FSequencerTrackFilter>& RHS) { return LHS->GetDisplayName().ToString() < RHS->GetDisplayName().ToString(); });
 }
 
+void SSequencer::UpdateOutlinerViewColumns()
+{
+	// Save updated column list in settings
+	TArray<FColumnVisibilitySetting> ColumnVisibilitySettings;
+
+	for (FSequencerOutlinerColumnVisibility ColumnVisibility : OutlinerColumnVisibilities)
+	{
+		ColumnVisibilitySettings.Add(FColumnVisibilitySetting(ColumnVisibility.Column->GetColumnName(), ColumnVisibility.bIsColumnVisible));
+	}
+
+	GetSequencerSettings()->SetOutlinerColumnVisibility(ColumnVisibilitySettings);
+
+	// Filter out hidden columns to create a list of visible columns for the outliner views
+	TArray<TSharedPtr<ISequencerOutlinerColumn>> VisibleColumns;
+	for (FSequencerOutlinerColumnVisibility ColumnVisibility : OutlinerColumnVisibilities)
+	{
+		if (ColumnVisibility.bIsColumnVisible)
+		{
+			VisibleColumns.Add(ColumnVisibility.Column);
+		}
+	}
+
+	// Update both Outliner Views with updated visible outliner columns
+	PinnedTreeView->SetOutlinerColumns(VisibleColumns);
+	TreeView->SetOutlinerColumns(VisibleColumns);
+}
+
+void SSequencer::InitializeOutlinerColumns()
+{
+	using namespace UE::Sequencer;
+
+	TSharedPtr<FSequencer> Sequencer = SequencerPtr.Pin();
+	if (!Sequencer.IsValid())
+	{
+		return;
+	}
+
+	const TMap<FName, TSharedPtr<ISequencerOutlinerColumn>>& RegisteredColumns = Sequencer->GetOutlinerColumns();
+	
+	// Retrieve previously saved column names and visibilities
+	TArray<FColumnVisibilitySetting> ColumnSettings = GetSequencerSettings()->GetOutlinerColumnSettings();
+	TSet<FName> ColumnNamesFoundInSettings;
+
+	// Add registered columns found in settings with their saved visibility state
+	for (const FColumnVisibilitySetting& ColumnVisibility : ColumnSettings)
+	{
+		const TSharedPtr<ISequencerOutlinerColumn>* OutlinerColumn = RegisteredColumns.Find(ColumnVisibility.ColumnName);
+		if (OutlinerColumn)
+		{
+			ColumnNamesFoundInSettings.Add(ColumnVisibility.ColumnName);
+			OutlinerColumnVisibilities.Add(FSequencerOutlinerColumnVisibility(*OutlinerColumn, ColumnVisibility.bIsVisible));
+		}
+	}
+
+	// Add registered columns not found in settings with their default visibility state
+	for (const TTuple<FName, TSharedPtr<ISequencerOutlinerColumn>>& RegisteredColumn : RegisteredColumns)
+	{
+		if (!ColumnNamesFoundInSettings.Contains(RegisteredColumn.Key))
+		{
+			OutlinerColumnVisibilities.Add(FSequencerOutlinerColumnVisibility(RegisteredColumn.Value));
+		}
+	}
+
+	UpdateOutlinerViewColumns();
+}
+
 /* SSequencer callbacks
  *****************************************************************************/
 
@@ -2102,6 +2172,9 @@ TSharedRef<SWidget> SSequencer::MakeViewMenu()
 		GetSequencerSettings()->SetZeroPadFrames(NewValue);
 	};
 
+	// Menu Entry for Outliner Column Visibilities
+	MenuBuilder.AddSubMenu(LOCTEXT("ColumnVisibilityHeader", "Columns"), FText::GetEmpty(), FNewMenuDelegate::CreateRaw(this, &SSequencer::FillColumnVisibilityMenu));
+
 	MenuBuilder.AddWidget(
 		SNew(SHorizontalBox)
 		+ SHorizontalBox::Slot()
@@ -2265,6 +2338,29 @@ void SSequencer::FillPlaybackSpeedMenu(FMenuBuilder& InMenuBarBuilder)
 	InMenuBarBuilder.EndSection();
 }
 
+void SSequencer::FillColumnVisibilityMenu(FMenuBuilder& InMenuBuilder)
+{
+	const bool bShouldCloseWindowAfterMenuSelection = true;
+	for (FSequencerOutlinerColumnVisibility& ColumnVisibility : OutlinerColumnVisibilities)
+	{
+		InMenuBuilder.AddMenuEntry(
+			ColumnVisibility.Column->GetColumnLabel(),
+			LOCTEXT("SetColumnVisibilityTooltip", "Enable or disable this outliner column"),
+			FSlateIcon(),
+			FUIAction(
+				FExecuteAction::CreateLambda([this, &ColumnVisibility] {
+					ColumnVisibility.bIsColumnVisible = !ColumnVisibility.bIsColumnVisible;
+					UpdateOutlinerViewColumns();
+					}),
+				FCanExecuteAction(),
+				FIsActionChecked::CreateLambda([this, ColumnVisibility] { return ColumnVisibility.bIsColumnVisible; })),
+				NAME_None,
+				EUserInterfaceActionType::ToggleButton
+			);
+	}
+	InMenuBuilder.EndSection();
+}
+
 void SSequencer::FillTimeDisplayFormatMenu(FMenuBuilder& MenuBuilder)
 {
 	TSharedPtr<FSequencer> Sequencer = SequencerPtr.Pin();
diff --git a/Engine/Source/Editor/Sequencer/Private/SSequencer.h b/Engine/Source/Editor/Sequencer/Private/SSequencer.h
--- a/Engine/Source/Editor/Sequencer/Private/SSequencer.h
+++ b/Engine/Source/Editor/Sequencer/Private/SSequencer.h
@@ -14,6 +14,7 @@
 #include "MovieSceneSequenceID.h"
 #include "ITimeSlider.h"
 #include "ISequencerModule.h"
+#include "ISequencerOutlinerColumn.h"
 #include "ToolMenu.h"
 #include "Framework/Commands/UICommandList.h"
 #include "Widgets/Input/NumericTypeInterface.h"
@@ -109,6 +110,26 @@ struct FSequencerBreadcrumb
 };
 
 
+/**
+ * Holds an outliner column and its visibility state
+ */
+struct FSequencerOutlinerColumnVisibility
+{
+	TSharedPtr<ISequencerOutlinerColumn> Column;
+	bool bIsColumnVisible = false;
+
+	FSequencerOutlinerColumnVisibility(TSharedPtr<ISequencerOutlinerColumn> InColumn)
+		: Column(InColumn)
+		, bIsColumnVisible(InColumn->IsColumnVisibleByDefault())
+	{}
+
+	FSequencerOutlinerColumnVisibility(TSharedPtr<ISequencerOutlinerColumn> InColumn, bool bInIsColumnVisible)
+		: Column(InColumn)
+		, bIsColumnVisible(bInIsColumnVisible)
+	{}
+};
+
+
 /**
  * A widget that holds a widget that is to be refocused on completion
  */
@@ -399,6 +420,10 @@ private:
 	/** Initalizes a list of all track filter objects */
 	void InitializeTrackFilters();
 
+
+	/** Initializes outliner column list from settings and SequencerCore */
+	void InitializeOutlinerColumns();
+
 	/** Handles key selection changes. */
 	void HandleKeySelectionChanged();
 
@@ -461,6 +486,9 @@ private:
 	/** Makes the playback speed menu for the toolbar. */
 	void FillPlaybackSpeedMenu(FMenuBuilder& InMenuBuilder);
 
+	/** Makes the column visibility menu for the toolbar. */
+	void FillColumnVisibilityMenu(FMenuBuilder& InMenuBuilder);
+
 	/** Return the current sequencer settings */ 
 	USequencerSettings* GetSequencerSettings() const;
 
@@ -495,6 +523,12 @@ private:
 	void OnEnableAllNodeGroupFilters(bool bEnableAll);
 	void OnNodeGroupFilterClicked(UMovieSceneNodeGroup* NodeGroup);
 
+	/**
+	 * Called when any outliner column's visibily is modified.
+	 * Updates SequencerSettings and visible outliner columns in Outliner View.
+	 */
+	void UpdateOutlinerViewColumns();
+
 	/**
 	* Called when the time snap interval changes.
 	*
@@ -709,6 +743,9 @@ private:
 	/** The fill coefficients of each column in the grid. */
 	float ColumnFillCoefficients[2];
 
+	/** List of registered outliner columns with their visibility states */
+	TArray<FSequencerOutlinerColumnVisibility> OutlinerColumnVisibilities;
+
 	TSharedPtr<class SSequencerSplitterOverlay> TreeViewSplitter;
 
 	/** Whether the active timer is currently registered */
diff --git a/Engine/Source/Editor/Sequencer/Private/SequencerSettings.cpp b/Engine/Source/Editor/Sequencer/Private/SequencerSettings.cpp
--- a/Engine/Source/Editor/Sequencer/Private/SequencerSettings.cpp
+++ b/Engine/Source/Editor/Sequencer/Private/SequencerSettings.cpp
@@ -993,3 +993,12 @@ void USequencerSettings::SetTrackFilterEnabled(const FString & TrackFilter, bool
 		}
 	}
 }
+
+void USequencerSettings::SetOutlinerColumnVisibility(const TArray<FColumnVisibilitySetting>& InHiddenColumns)
+{
+	if (InHiddenColumns != HiddenColumns)
+	{
+		HiddenColumns = InHiddenColumns;
+		SaveConfig();
+	}
+}
diff --git a/Engine/Source/Editor/Sequencer/Public/MovieSceneOutlinerColumn.h b/Engine/Source/Editor/Sequencer/Public/MovieSceneOutlinerColumn.h
--- a/Engine/Source/Editor/Sequencer/Public/MovieSceneOutlinerColumn.h
+++ b/Engine/Source/Editor/Sequencer/Public/MovieSceneOutlinerColumn.h
@@ -23,5 +23,4 @@ public:
 
 	virtual bool SupportsSequence(UMovieSceneSequence* InSequence) const override { return true; }
 
-	virtual TSharedRef<SWidget> CreateColumnWidget(UE::Sequencer::TViewModelPtr<UE::Sequencer::IOutlinerExtension> InOutlinerExtension) const = 0;
 };
diff --git a/Engine/Source/Editor/Sequencer/Public/SequencerSettings.h b/Engine/Source/Editor/Sequencer/Public/SequencerSettings.h
--- a/Engine/Source/Editor/Sequencer/Public/SequencerSettings.h
+++ b/Engine/Source/Editor/Sequencer/Public/SequencerSettings.h
@@ -80,6 +80,31 @@ public:
 	}
 };
 
+/** Struct for storing reorderable and hidden/visible outliner columns */
+USTRUCT(BlueprintType)
+struct FColumnVisibilitySetting
+{
+	GENERATED_BODY()
+
+	UPROPERTY(EditAnywhere, Category=General)
+	FName ColumnName;
+
+	UPROPERTY(EditAnywhere, Category=General)
+	bool bIsVisible;
+
+	bool operator==(const FColumnVisibilitySetting& Other) const
+	{
+		return ColumnName == Other.ColumnName && bIsVisible == Other.bIsVisible;
+	}
+
+	FColumnVisibilitySetting(FName InColumnName, bool InbIsVisible)
+		: ColumnName(InColumnName)
+		, bIsVisible(InbIsVisible)
+	{}
+
+	FColumnVisibilitySetting()
+	{}
+};
 
 /** Serializable options for sequencer. */
 UCLASS(config=EditorPerProjectUserSettings, PerObjectConfig)
@@ -420,6 +445,11 @@ public:
 	/** Sets whether the track filter should be enabled/disabled */
 	void SetTrackFilterEnabled(const FString& TrackFilter, bool bEnabled);
 
+	/** Get outliner column visibility in display order */
+	TArray<FColumnVisibilitySetting> GetOutlinerColumnSettings() const { return HiddenColumns; }
+	/** Sets the visibility of outliner columns in display order */
+	void SetOutlinerColumnVisibility(const TArray<FColumnVisibilitySetting>& InHiddenColumns);
+
 protected:
 
 	/** The auto change mode (auto-key, auto-track or none). */
@@ -672,6 +702,10 @@ protected:
 	UPROPERTY(config, EditAnywhere, Category = General)
 	TArray<FString> TrackFilters;
 
+	/** List of all columns and their visibility, in the order to be displayed in the outliner view */
+	UPROPERTY(config, EditAnywhere, Category = General)
+	TArray<FColumnVisibilitySetting> HiddenColumns;
+
 	FOnEvaluateSubSequencesInIsolationChanged OnEvaluateSubSequencesInIsolationChangedEvent;
 	FOnShowSelectedNodesOnlyChanged OnShowSelectedNodesOnlyChangedEvent;
 	FOnAllowEditsModeChanged OnAllowEditsModeChangedEvent;
diff --git a/Engine/Source/Editor/SequencerCore/Private/MVVM/Views/SOutlinerColumnBorder.cpp b/Engine/Source/Editor/SequencerCore/Private/MVVM/Views/SOutlinerColumnBorder.cpp
new file mode 100644
--- /dev/null
+++ b/Engine/Source/Editor/SequencerCore/Private/MVVM/Views/SOutlinerColumnBorder.cpp
@@ -0,0 +1,181 @@
+// Copyright Epic Games, Inc. All Rights Reserved.
+
+#include "MVVM/Views/SOutlinerColumnBorder.h"
+
+#include "MVVM/Extensions/IOutlinerExtension.h"
+#include "MVVM/ViewModelPtr.h"
+#include "MVVM/ViewModels/EditorViewModel.h"
+#include "MVVM/ViewModels/OutlinerViewModel.h"
+#include "MVVM/ViewModels/ViewModelIterators.h"
+#include "Styling/StyleColors.h"
+
+namespace UE::Sequencer
+{
+
+void SOutlinerColumnInnerBorder::Construct(const FArguments& InArgs, const FCreateOutlinerColumnParams& InParams)
+{
+	bIsMouseOverInnerBorder = false;
+
+	WeakOutlinerExtension = InParams.OutlinerExtension;
+	WeakEditor = InParams.Editor;
+
+	BackgroundBrush = FAppStyle::GetBrush("Sequencer.AnimationOutliner.DefaultBorder");
+
+	// Size of outliner column widgets is currently 12x12
+	TSharedRef<SWidget>	FinalWidget = SNew(SBox)
+		.WidthOverride(12.0f)
+		.HeightOverride(12.0f)
+		.HAlign(HAlign_Center)
+		.VAlign(VAlign_Center)
+		[
+			SNew(SBorder)
+			.VAlign(VAlign_Center)
+			.BorderImage(this, &SOutlinerColumnInnerBorder::GetBorderImage)
+			.BorderBackgroundColor(this, &SOutlinerColumnInnerBorder::GetBackgroundTint)
+			.Padding(FMargin(0.f))
+			[
+				InArgs._Content.Widget
+			]
+		];
+
+	ChildSlot
+	[
+		FinalWidget
+	];
+}
+
+void SOutlinerColumnInnerBorder::OnMouseEnter(const FGeometry& MyGeometry, const FPointerEvent& MouseEvent)
+{
+	bIsMouseOverInnerBorder = true;
+}
+
+void SOutlinerColumnInnerBorder::OnMouseLeave(const FPointerEvent& MouseEvent)
+{
+	bIsMouseOverInnerBorder = false;
+}
+
+const FSlateBrush* SOutlinerColumnInnerBorder::GetBorderImage() const
+{
+	return BackgroundBrush;
+}
+
+FSlateColor SOutlinerColumnInnerBorder::GetBackgroundTint() const
+{
+	TSharedPtr<FEditorViewModel> Editor = WeakEditor.Pin();
+	TViewModelPtr<IOutlinerExtension> OutlinerItem = WeakOutlinerExtension.Pin();
+
+	EOutlinerSelectionState SelectionState = OutlinerItem->GetSelectionState();
+
+	bool bIsOutlinerItemSelectedDirectly = EnumHasAnyFlags(SelectionState, EOutlinerSelectionState::SelectedDirectly);
+	bool bIsOutlinerItemHovered = Editor->GetOutliner()->GetHoveredItem() == OutlinerItem;
+
+	// float Opacity = ((bIsOutlinerItemSelectedDirectly || bIsOutlinerItemHovered) && !bIsMouseOverInnerBorder) ? 0.f : .4f;
+	float Opacity = .4f;
+	return FLinearColor(0.f, 0.f, 0.f, Opacity);
+}
+
+
+void SOutlinerColumnBorder::Construct(const FArguments& InArgs, const FCreateOutlinerColumnParams& InParams)
+{
+	WeakOutlinerExtension = InParams.OutlinerExtension;
+	WeakEditor = InParams.Editor;
+
+	BackgroundBrush = FAppStyle::GetBrush("Sequencer.AnimationOutliner.DefaultBorder");
+
+	TSharedRef<SWidget>	FinalWidget = SNew(SBorder)
+		.VAlign(VAlign_Center)
+		.BorderImage(this, &SOutlinerColumnBorder::GetBorderImage)
+		.BorderBackgroundColor(this, &SOutlinerColumnBorder::GetBackgroundTint)
+		.VAlign(VAlign_Center)
+		.HAlign(HAlign_Center)
+		.Padding(FMargin(1.0f))
+		[
+			SNew(SOutlinerColumnInnerBorder, InParams)
+			[
+				InArgs._Content.Widget
+			]
+		];
+
+	ChildSlot
+	[
+		FinalWidget
+	];
+}
+
+void SOutlinerColumnBorder::OnMouseEnter(const FGeometry& MyGeometry, const FPointerEvent& MouseEvent)
+{
+	TViewModelPtr<IOutlinerExtension> DataModel = WeakOutlinerExtension.Pin();
+	TSharedPtr<FEditorViewModel> Editor = WeakEditor.Pin();
+	if (DataModel && Editor)
+	{
+		Editor->GetOutliner()->SetHoveredItem(DataModel);
+	}
+	SWidget::OnMouseEnter(MyGeometry, MouseEvent);
+}
+
+void SOutlinerColumnBorder::OnMouseLeave(const FPointerEvent& MouseEvent)
+{
+	TSharedPtr<FEditorViewModel> Editor = WeakEditor.Pin();
+	if (Editor)
+	{
+		Editor->GetOutliner()->SetHoveredItem(nullptr);
+	}
+	SWidget::OnMouseLeave(MouseEvent);
+}
+
+const FSlateBrush* SOutlinerColumnBorder::GetBorderImage() const
+{
+	return BackgroundBrush;
+}
+
+FSlateColor SOutlinerColumnBorder::GetBackgroundTint() const
+{
+	TSharedPtr<FEditorViewModel> Editor = WeakEditor.Pin();
+	TViewModelPtr<IOutlinerExtension> OutlinerItem = WeakOutlinerExtension.Pin();
+
+	if (!Editor || !OutlinerItem)
+	{
+		return FLinearColor(0.f, 0.f, 0.f, 0.f);
+	}
+
+	if (!OutlinerItem->HasBackground())
+	{
+		return FLinearColor(0.f, 0.f, 0.f, 0.f);
+	}
+
+	EOutlinerSelectionState SelectionState = OutlinerItem->GetSelectionState();
+
+	if (EnumHasAnyFlags(SelectionState, EOutlinerSelectionState::SelectedDirectly))
+	{
+		return FStyleColors::Select;
+	}
+	if (EnumHasAnyFlags(SelectionState, EOutlinerSelectionState::HasSelectedKeys | EOutlinerSelectionState::HasSelectedTrackAreaItems))
+	{
+		return FStyleColors::Header;
+	}
+
+	// If this is collapsed but has any children with selected keys or sections, we report that state on the parent
+	if (!OutlinerItem->IsExpanded())
+	{
+		for (TViewModelPtr<IOutlinerExtension> Child : OutlinerItem.AsModel()->GetDescendantsOfType<IOutlinerExtension>())
+		{
+			if (EnumHasAnyFlags(SelectionState, EOutlinerSelectionState::HasSelectedKeys | EOutlinerSelectionState::HasSelectedTrackAreaItems))
+			{
+				return FStyleColors::Header;
+			}
+		}
+	}
+
+	// Use same hovered and default colors as SOutlinerItemViewBase
+	if (Editor->GetOutliner()->GetHoveredItem() == OutlinerItem)
+	{
+		return FLinearColor(FColor(72, 72, 72, 255));
+	}
+
+	return FLinearColor(FColor(62, 62, 62, 255));
+}
+
+} // namespace UE::Sequencer
+
+#undef LOCTEXT_NAMESPACE
+
diff --git a/Engine/Source/Editor/SequencerCore/Private/MVVM/Views/SOutlinerView.cpp b/Engine/Source/Editor/SequencerCore/Private/MVVM/Views/SOutlinerView.cpp
--- a/Engine/Source/Editor/SequencerCore/Private/MVVM/Views/SOutlinerView.cpp
+++ b/Engine/Source/Editor/SequencerCore/Private/MVVM/Views/SOutlinerView.cpp
@@ -12,6 +12,7 @@
 #include "InputCoreTypes.h"
 #include "Internationalization/Internationalization.h"
 #include "Internationalization/Text.h"
+#include "ISequencerOutlinerColumn.h"
 #include "Layout/ChildrenBase.h"
 #include "Layout/Geometry.h"
 #include "Layout/Margin.h"
@@ -29,6 +30,7 @@
 #include "MVVM/ViewModels/ViewModelIterators.h"
 #include "MVVM/Selection/SequencerCoreSelection.h"
 #include "MVVM/Selection/SequencerOutlinerSelection.h"
+#include "MVVM/Views/SOutlinerColumnBorder.h"
 #include "MVVM/Views/STrackAreaView.h"
 #include "MVVM/Views/STrackLane.h"
 #include "MVVM/Views/TreeViewTraits.h"
@@ -47,6 +49,7 @@
 #include "Templates/Tuple.h"
 #include "Types/SlateStructs.h"
 #include "Widgets/Layout/SBox.h"
+#include "Widgets/Layout/SSeparator.h"
 #include "Widgets/SNullWidget.h"
 #include "Widgets/Views/SHeaderRow.h"
 
@@ -270,7 +273,7 @@ void SOutlinerView::Construct(const FArguments& InArgs, TWeakPtr<FOutlinerViewMo
 
 	HeaderRow = SNew(SHeaderRow).Visibility(EVisibility::Collapsed);
 
-	SetupColumns(InArgs);
+	GenerateOutlinerColumns();
 
 	WeakOutliner.Pin()->OnRefreshed.AddSP(this, &SOutlinerView::Refresh);
 
@@ -467,11 +470,60 @@ float SOutlinerView::VirtualToPhysical(float InVirtual) const
 	return InVirtual;
 }
 
-void SOutlinerView::SetupColumns(const FArguments& InArgs)
+void SOutlinerView::GenerateOutlinerColumns()
 {
+	// Clear columns to ensure consistent order when building UI from Map
+	HeaderRow->ClearColumns();
+	Columns.Empty();
+	Columns.Compact();
+
 	TSharedPtr<FEditorViewModel> EditorViewModel = WeakOutliner.Pin()->GetEditor();
 
-	// Define a column for the Outliner
+	// Create Visible Outliner Columns (Pin/Mute/Solo...)
+	for (TSharedPtr<ISequencerOutlinerColumn> Column : OutlinerColumns)
+	{
+		auto GenerateToolColumn = [=](const TWeakViewModelPtr<IOutlinerExtension>& InWeakModel, const TSharedRef<SOutlinerViewRow>& InRow) -> TSharedRef<SWidget>
+		{
+			if (TViewModelPtr<IOutlinerExtension> OutlinerExtension = InWeakModel.ImplicitPin())
+			{
+				if (!OutlinerExtension.AsModel()->IsA<FOutlinerSpacer>())
+				{
+					return SNew(SOutlinerColumnBorder, FCreateOutlinerColumnParams(OutlinerExtension, EditorViewModel))
+						[	
+							Column->CreateColumnWidget(FCreateOutlinerColumnParams(OutlinerExtension, EditorViewModel))
+						];
+				}
+				else
+				{
+					return SNew(SBox).HeightOverride(10.f);
+				}
+			}
+
+			ensureMsgf(false, TEXT("Attempting to create an outliner column for a view model that is either dead, or not an outliner item."));
+			return SNew(SBox).HeightOverride(10.f);
+		};
+		Columns.Add(Column->GetColumnName(), FOutlinerViewColumn(GenerateToolColumn, 16.f, true));
+	}
+
+	// If there are any outliner columns, add a spacer between them and the outliner view
+	if (OutlinerColumns.Num() > 0)
+	{
+		auto GenerateSpacer = [=](const TWeakViewModelPtr<IOutlinerExtension>& InWeakModel, const TSharedRef<SOutlinerViewRow>& InRow) -> TSharedRef<SWidget>
+		{
+			return SNew(SHorizontalBox)
+				+ SHorizontalBox::Slot()
+				.AutoWidth()
+				[
+					SNew(SSeparator)
+					.Orientation(Orient_Vertical)
+					.Thickness(1.0f)
+					.SeparatorImage(FAppStyle::Get().GetBrush("Separator"))
+				];
+		};
+		Columns.Add(FName(TEXT("OutlinerColumnSpacer")), FOutlinerViewColumn(GenerateSpacer, 1.f, true));
+	}
+
+	// Create Column for Outliner View (Track name / Keying)
 	auto GenerateOutliner = [=](const TWeakViewModelPtr<IOutlinerExtension>& InWeakModel, const TSharedRef<SOutlinerViewRow>& InRow) -> TSharedRef<SWidget>
 	{
 		if (TSharedPtr<IOutlinerExtension> OutlinerItem = InWeakModel.ImplicitPin())
@@ -490,10 +542,25 @@ void SOutlinerView::SetupColumns(const FArguments& InArgs)
 	{
 		if (Pair.Key != TrackNameColumn)
 		{
-			HeaderRow->AddColumn(
-				SHeaderRow::Column(Pair.Key)
-				.FillWidth(Pair.Value.Width)
-			);
+			if (Pair.Value.bIsFixedWidth)
+			{
+				TOptional<float> FixedWidth;
+				if (Pair.Value.Width.IsSet())
+				{
+					FixedWidth = Pair.Value.Width.Get();
+				}
+				HeaderRow->AddColumn(
+					SHeaderRow::Column(Pair.Key)
+					.FixedWidth(FixedWidth)
+				);
+			}
+			else
+			{
+				HeaderRow->AddColumn(
+					SHeaderRow::Column(Pair.Key)
+					.FillWidth(Pair.Value.Width)
+				);
+			}
 		}
 	}
 }
@@ -516,6 +583,13 @@ void SOutlinerView::AddPinnedTreeView(TSharedPtr<SOutlinerView> PinnedTreeView)
 	PinnedTreeView->SetPrimaryTreeView(SharedThis(this));
 }
 
+void SOutlinerView::SetOutlinerColumns(const TArray<TSharedPtr<ISequencerOutlinerColumn>>& InOutlinerColumns)
+{
+	// Reset the way rows are constructed with an updated list of Outliner Columns
+	OutlinerColumns = InOutlinerColumns;
+	GenerateOutlinerColumns();
+}
+
 void SOutlinerView::OnRightMouseButtonDown(const FPointerEvent& MouseEvent)
 {
 	STreeView::OnRightMouseButtonDown(MouseEvent);
diff --git a/Engine/Source/Editor/SequencerCore/Public/ISequencerOutlinerColumn.h b/Engine/Source/Editor/SequencerCore/Public/ISequencerOutlinerColumn.h
--- a/Engine/Source/Editor/SequencerCore/Public/ISequencerOutlinerColumn.h
+++ b/Engine/Source/Editor/SequencerCore/Public/ISequencerOutlinerColumn.h
@@ -2,30 +2,51 @@
 
 #pragma once
 
+#include "Internationalization/Text.h"
 #include "MVVM/ViewModelPtr.h"
 #include "Templates/SubclassOf.h"
 
+class SWidget;
 class UMovieSceneSequence;
 
-namespace UE
-{
-namespace Sequencer
-{
+namespace UE::Sequencer 
+{ 
 class IOutlinerExtension;
+class FEditorViewModel;
 }
+
+namespace UE::Sequencer
+{
+
+/** Parameters for creating an outliner column widget. */
+struct SEQUENCERCORE_API FCreateOutlinerColumnParams
+{
+
+	FCreateOutlinerColumnParams(const TViewModelPtr<IOutlinerExtension> InOutlinerExtension, const TSharedPtr<FEditorViewModel> InEditor)
+		: OutlinerExtension(InOutlinerExtension)
+		, Editor(InEditor)
+	{}
+
+	const TViewModelPtr<IOutlinerExtension> OutlinerExtension;
+	const TSharedPtr<FEditorViewModel> Editor;
+};
+
 }
 
 /**
- * Interface for sequencer outliner columns.
- */
-class ISequencerOutlinerColumn
+* Interface for sequencer outliner columns.
+*/
+class SEQUENCERCORE_API ISequencerOutlinerColumn
 {
-	
+
 public:
 
 	/* Gets the unique name of the column to use in the SOutlinerView registry and context menus */
 	virtual FName GetColumnName() const = 0;
 
+	/* Gets the text to display in the UI for visibility */
+	virtual FText GetColumnLabel() const = 0;
+
 	/* The default visibility state of this column when loaded for the first time */
 	virtual bool IsColumnVisibleByDefault() const { return true; }
 
@@ -33,11 +54,11 @@ public:
 	virtual bool SupportsSequence(UMovieSceneSequence* InSequence) const = 0;
 
 	/* Gets the widget created for each item within the SOutlinerView, column widgets must be fixed width */
-	virtual TSharedRef<SWidget> CreateColumnWidget(UE::Sequencer::TViewModelPtr<UE::Sequencer::IOutlinerExtension>InOutlinerExtension) const = 0;
+	virtual TSharedRef<SWidget> CreateColumnWidget(const UE::Sequencer::FCreateOutlinerColumnParams& InParams) const = 0;
 
 public:
 
 	/** Virtual destructor. */
 	virtual ~ISequencerOutlinerColumn() { }
-	
+
 };
diff --git a/Engine/Source/Editor/SequencerCore/Public/MVVM/Views/SOutlinerColumnBorder.h b/Engine/Source/Editor/SequencerCore/Public/MVVM/Views/SOutlinerColumnBorder.h
new file mode 100644
--- /dev/null
+++ b/Engine/Source/Editor/SequencerCore/Public/MVVM/Views/SOutlinerColumnBorder.h
@@ -0,0 +1,107 @@
+// Copyright Epic Games, Inc. All Rights Reserved.
+
+#pragma once
+
+#include "ISequencerOutlinerColumn.h"
+#include "MVVM/ViewModelPtr.h"
+#include "Widgets/SCompoundWidget.h"
+
+namespace UE::Sequencer
+{
+
+class FEditorViewModel;
+class IOutlinerExtension;
+
+/**
+* The widget for the inner, darker border of a column widget that handles coloring / sizing
+*/
+class SOutlinerColumnInnerBorder
+	: public SCompoundWidget
+{
+
+public:
+
+	SLATE_BEGIN_ARGS(SOutlinerColumnInnerBorder) {}
+		SLATE_DEFAULT_SLOT(FArguments, Content)
+	SLATE_END_ARGS()
+
+	void Construct(
+		const FArguments& InArgs,
+		const FCreateOutlinerColumnParams& InParams
+	);
+
+private:
+
+	/** Determine whether or not the mouse is directly over this widget for background color */
+	void OnMouseEnter(const FGeometry& MyGeometry, const FPointerEvent& MouseEvent) override;
+	void OnMouseLeave(const FPointerEvent& MouseEvent) override;
+
+	/**
+	* @return The border image to show in the tree node.
+	*/
+	const FSlateBrush* GetBorderImage() const;
+
+	/**
+	* @return The inner background border color and opacity, darker when mouse is directly over inner widget.
+	*/
+	FSlateColor GetBackgroundTint() const;
+
+	/** Used to determine color if outliner item is directly selected */
+	TWeakViewModelPtr<IOutlinerExtension> WeakOutlinerExtension;
+
+	/** Used to determine color if outliner item is being hovered */
+	TWeakPtr<FEditorViewModel> WeakEditor;
+
+	/** Brush used for inner border */
+	const FSlateBrush* BackgroundBrush;
+
+	/** Is the mouse directly over this widget */
+	bool bIsMouseOverInnerBorder = false;
+};
+
+/**
+* A widget for containing a column widget with the correct style / border
+*/
+class SOutlinerColumnBorder
+	: public SCompoundWidget
+{
+
+public:
+
+	SLATE_BEGIN_ARGS(SOutlinerColumnBorder) {}
+		SLATE_DEFAULT_SLOT(FArguments, Content)
+	SLATE_END_ARGS()
+
+	void Construct(
+		const FArguments& InArgs, 
+		const FCreateOutlinerColumnParams& InParams
+	);
+
+private:
+
+	/** Notifies the outliner item that the mouse is within the row for hovered state */
+	void OnMouseEnter(const FGeometry& MyGeometry, const FPointerEvent& MouseEvent) override;
+	void OnMouseLeave(const FPointerEvent& MouseEvent) override;
+
+	/**
+	* @return The border image to show in the tree node.
+	*/
+	const FSlateBrush* GetBorderImage() const;
+
+	/**
+	* @return The tint to apply to the border image
+	*/
+	FSlateColor GetBackgroundTint() const;
+
+	/** Used to determine color if outliner item is directly selected */
+	TWeakViewModelPtr<IOutlinerExtension> WeakOutlinerExtension;
+
+	/** Used to determine color if outliner item is being hovered */
+	TWeakPtr<FEditorViewModel> WeakEditor;
+
+	/** Brush used for border */
+	const FSlateBrush* BackgroundBrush;
+};
+
+} // namespace UE::Sequencer
+
diff --git a/Engine/Source/Editor/SequencerCore/Public/MVVM/Views/SOutlinerView.h b/Engine/Source/Editor/SequencerCore/Public/MVVM/Views/SOutlinerView.h
--- a/Engine/Source/Editor/SequencerCore/Public/MVVM/Views/SOutlinerView.h
+++ b/Engine/Source/Editor/SequencerCore/Public/MVVM/Views/SOutlinerView.h
@@ -32,6 +32,7 @@ class FPaintArgs;
 class FSlateRect;
 class FSlateWindowElementList;
 class FWidgetStyle;
+class ISequencerOutlinerColumn;
 class ITableRow;
 class SHeaderRow;
 class SScrollBar;
@@ -65,16 +66,18 @@ struct SEQUENCERCORE_API FOutlinerViewColumn
 {
 	typedef TFunction<TSharedRef<SWidget>(TViewModelPtr<IOutlinerExtension>, const TSharedRef<SOutlinerViewRow>&)> FOnGenerate;
 
-	FOutlinerViewColumn(const FOnGenerate& InOnGenerate, const TAttribute<float>& InWidth) : Generator(InOnGenerate), Width(InWidth) {}
-	FOutlinerViewColumn(FOnGenerate&& InOnGenerate, const TAttribute<float>& InWidth) : Generator(MoveTemp(InOnGenerate)), Width(InWidth) {}
+	FOutlinerViewColumn(const FOnGenerate& InOnGenerate, const TAttribute<float>& InWidth, const bool InIsFixedWidth) : Generator(InOnGenerate), Width(InWidth), bIsFixedWidth(InIsFixedWidth) {}
+	FOutlinerViewColumn(const FOnGenerate& InOnGenerate, const TAttribute<float>& InWidth) : Generator(InOnGenerate), Width(InWidth), bIsFixedWidth(false) {}
+	FOutlinerViewColumn(FOnGenerate&& InOnGenerate, const TAttribute<float>& InWidth) : Generator(MoveTemp(InOnGenerate)), Width(InWidth), bIsFixedWidth(false) {}
 
 	/** Function used to generate a cell for this column */
 	FOnGenerate Generator;
 	/** Attribute specifying the width of this column */
 	TAttribute<float> Width;
+	/** Attribute specifying whether or not to use the Width in percentage or fixed-width values */
+	bool bIsFixedWidth = false;
 };
 
-
 /** The tree view used in the sequencer */
 class SEQUENCERCORE_API SOutlinerView 
 	: public STreeView<TWeakViewModelPtr<IOutlinerExtension>>
@@ -224,6 +227,9 @@ public:
 	/** Set whether this TreeView should show only pinned nodes or only non-pinned nodes  */
 	void SetShowPinned(bool bShowPinned) { bShowPinnedNodes = bShowPinned; }
 
+	/** Updates the list of visible outliner columns and regenerates columns in the outliner view */
+	void SetOutlinerColumns(const TArray<TSharedPtr<ISequencerOutlinerColumn>>& InOutlinerColumns);
+
 protected:
 
 	/** Linear, sorted array of nodes that we currently have generated widgets for */
@@ -231,8 +237,8 @@ protected:
 
 protected:
 
-	/** Populate the map of column definitions, and add relevant columns to the header row */
-	void SetupColumns(const FArguments& InArgs);
+	/** Populate the map of column definitions, and add relevant columns to the header row. Must be called when outliner columns change */
+	void GenerateOutlinerColumns();
 
 	FReply OnDragRow(const FGeometry& InGeometry, const FPointerEvent& InPointerEvent, TSharedRef<SOutlinerViewRow> InRow);
 
@@ -264,6 +270,9 @@ private:
 	/** The SOutlinerView object this SOutlinerView is pinned to, or nullptr if not pinned */
 	TWeakPtr<SOutlinerView> PrimaryTreeView;
 
+	/** Visible Outliner columns to display in the outliner view */
+	TArray<TSharedPtr<ISequencerOutlinerColumn>> OutlinerColumns;
+
 	float VirtualTop;
 
 	/** When true, the tree selection is being updated from a change in the sequencer selection. */
