Sequencer: Track Gutter UI Framework, Outliner Column Creation and Settings

diff --git a/Engine/Source/Editor/SequencerCore/Private/MVVM/Views/SOutlinerColumnBorder.cpp b/Engine/Source/Editor/SequencerCore/Private/MVVM/Views/SOutlinerColumnBorder.cpp
--- a/Engine/Source/Editor/SequencerCore/Private/MVVM/Views/SOutlinerColumnBorder.cpp
+++ b/Engine/Source/Editor/SequencerCore/Private/MVVM/Views/SOutlinerColumnBorder.cpp
@@ -177,5 +177,3 @@ FSlateColor SOutlinerColumnBorder::GetBackgroundTint() const
 
 } // namespace UE::Sequencer
 
-#undef LOCTEXT_NAMESPACE
-
