Sequencer: Fixed hotspot changed delegate not being broadcast for some functions

This corrects the event only being broadcast for the SetHotspot call from CL 24179637.

diff --git a/Engine/Source/Editor/SequencerCore/Private/MVVM/ViewModels/TrackAreaViewModel.cpp b/Engine/Source/Editor/SequencerCore/Private/MVVM/ViewModels/TrackAreaViewModel.cpp
--- a/Engine/Source/Editor/SequencerCore/Private/MVVM/ViewModels/TrackAreaViewModel.cpp
+++ b/Engine/Source/Editor/SequencerCore/Private/MVVM/ViewModels/TrackAreaViewModel.cpp
@@ -84,6 +84,8 @@ void FTrackAreaViewModel::AddHotspot(TSharedPtr<ITrackAreaHotspot> NewHotspot)
 	// We always set a hotspot if there is none
 	if (bHotspotLocked == false)
 	{
+		TSharedPtr<ITrackAreaHotspot> PreviousHotspot = GetHotspot();
+
 		HotspotStack.Push(NewHotspot);
 		HotspotStack.Sort([](TSharedPtr<ITrackAreaHotspot> A, TSharedPtr<ITrackAreaHotspot> B){
 			return A->Priority() > B->Priority();
@@ -91,6 +93,12 @@ void FTrackAreaViewModel::AddHotspot(TSharedPtr<ITrackAreaHotspot> NewHotspot)
 
 		// Simulate an update-on-hover for the new hotspot to ensure that any hover behavior doesn't have to wait until the next frame
 		NewHotspot->UpdateOnHover(*this);
+
+		TSharedPtr<ITrackAreaHotspot> CurrentHotspot = GetHotspot();
+		if (PreviousHotspot != CurrentHotspot)
+		{
+			OnHotspotChangedDelegate.Broadcast(CurrentHotspot);
+		}
 	}
 }
 
@@ -98,9 +106,17 @@ void FTrackAreaViewModel::RemoveHotspot(FViewModelTypeID Type)
 {
 	if (bHotspotLocked == false)
 	{
+		TSharedPtr<ITrackAreaHotspot> PreviousHotspot = GetHotspot();
+
 		HotspotStack.RemoveAll([Type](TSharedPtr<ITrackAreaHotspot> In){
 			return In->CastRaw(Type) != nullptr;
 		});
+
+		TSharedPtr<ITrackAreaHotspot> CurrentHotspot = GetHotspot();
+		if (PreviousHotspot != CurrentHotspot)
+		{
+			OnHotspotChangedDelegate.Broadcast(CurrentHotspot);
+		}
 	}
 }
 
@@ -109,6 +125,7 @@ void FTrackAreaViewModel::ClearHotspots()
 	if (bHotspotLocked == false)
 	{
 		HotspotStack.Empty();
+		OnHotspotChangedDelegate.Broadcast(nullptr);
 	}
 }
 
