Sequencer: Only allow dropping onto spacers from sequencer outilner drag and drop operations, ie. not actors from the outliner

diff --git a/Engine/Source/Editor/SequencerCore/Private/MVVM/ViewModels/OutlinerSpacer.cpp b/Engine/Source/Editor/SequencerCore/Private/MVVM/ViewModels/OutlinerSpacer.cpp
--- a/Engine/Source/Editor/SequencerCore/Private/MVVM/ViewModels/OutlinerSpacer.cpp
+++ b/Engine/Source/Editor/SequencerCore/Private/MVVM/ViewModels/OutlinerSpacer.cpp
@@ -4,6 +4,7 @@
 
 #include "HAL/PlatformCrt.h"
 #include "MVVM/ViewModelPtr.h"
+#include "MVVM/ViewModels/OutlinerViewModelDragDropOp.h"
 #include "MVVM/ViewModels/ViewModelIterators.h"
 #include "Math/NumericLimits.h"
 #include "Misc/Attribute.h"
@@ -75,6 +76,12 @@ void FOutlinerSpacer::SetCustomOrder(int32 InCustomOrder)
 
 TOptional<EItemDropZone> FOutlinerSpacer::CanAcceptDrop(const FViewModelPtr& TargetModel, const FDragDropEvent& DragDropEvent, EItemDropZone InItemDropZone)
 {
+	TSharedPtr<FOutlinerViewModelDragDropOp> DragDropOp = DragDropEvent.GetOperationAs<FOutlinerViewModelDragDropOp>();
+	if (!DragDropOp)
+	{
+		return TOptional<EItemDropZone>();
+	}
+
 	return EItemDropZone::AboveItem;
 }
 
