Sequencer: Outliner Column fixed height caused issues with alignment between Outliner View and Track View

#rnx

diff --git a/Engine/Source/Editor/SequencerCore/Private/MVVM/Views/SOutlinerColumnBorder.cpp b/Engine/Source/Editor/SequencerCore/Private/MVVM/Views/SOutlinerColumnBorder.cpp
--- a/Engine/Source/Editor/SequencerCore/Private/MVVM/Views/SOutlinerColumnBorder.cpp
+++ b/Engine/Source/Editor/SequencerCore/Private/MVVM/Views/SOutlinerColumnBorder.cpp
@@ -8,6 +8,7 @@
 #include "MVVM/ViewModels/OutlinerViewModel.h"
 #include "MVVM/ViewModels/ViewModelIterators.h"
 #include "Styling/StyleColors.h"
+#include "Widgets/Layout/SScaleBox.h"
 
 namespace UE::Sequencer
 {
@@ -21,22 +22,29 @@ void SOutlinerColumnInnerBorder::Construct(const FArguments& InArgs, const FCrea
 
 	BackgroundBrush = FAppStyle::GetBrush("Sequencer.Column.OutlinerColumnBox");
 
-	// Size of outliner column widgets is currently 12x12
-	TSharedRef<SWidget>	FinalWidget = SNew(SBox)
-		.WidthOverride(12.0f)
-		.HeightOverride(12.0f)
+	TSharedPtr<SScaleBox> UniformScaleBox;
+
+	// Size of outliner column widgets stretch to desired height and are usually 12x12 based off padding and fixed width of columns
+	TSharedRef<SWidget>	FinalWidget = SNew(SVerticalBox)
+		+ SVerticalBox::Slot()
+		.AutoHeight()
 		.HAlign(HAlign_Center)
 		.VAlign(VAlign_Center)
 		[
-			SNew(SBorder)
-			.VAlign(VAlign_Center)
-			.BorderImage(this, &SOutlinerColumnInnerBorder::GetBorderImage)
-			.Padding(FMargin(1.f))
+			SAssignNew(UniformScaleBox, SScaleBox)
 			[
-				InArgs._Content.Widget
+				SNew(SBorder)
+				.VAlign(VAlign_Center)
+				.BorderImage(this, &SOutlinerColumnInnerBorder::GetBorderImage)
+				.Padding(FMargin(1.f))
+				[
+					InArgs._Content.Widget
+				]
 			]
 		];
 
+	UniformScaleBox->SetStretch(EStretch::ScaleToFit);
+
 	ChildSlot
 	[
 		FinalWidget
@@ -89,7 +97,7 @@ void SOutlinerColumnBorder::Construct(const FArguments& InArgs, const FCreateOut
 		.BorderBackgroundColor(this, &SOutlinerColumnBorder::GetBackgroundTint)
 		.VAlign(VAlign_Center)
 		.HAlign(HAlign_Center)
-		.Padding(FMargin(0.0f))
+		.Padding(FMargin(2.0f))
 		[
 			InnerWidget.ToSharedRef()
 		];
diff --git a/Engine/Source/Editor/SequencerCore/Private/MVVM/Views/SOutlinerView.cpp b/Engine/Source/Editor/SequencerCore/Private/MVVM/Views/SOutlinerView.cpp
--- a/Engine/Source/Editor/SequencerCore/Private/MVVM/Views/SOutlinerView.cpp
+++ b/Engine/Source/Editor/SequencerCore/Private/MVVM/Views/SOutlinerView.cpp
@@ -525,15 +525,7 @@ void SOutlinerView::UpdateTrackGutterColumns(TSharedPtr<FEditorViewModel> InEdit
 	// Spacer column callback, generates a line to separate the track gutter from the outliner
 	auto GenerateSpacer = [=](const TWeakViewModelPtr<IOutlinerExtension>& InWeakModel, const TSharedRef<SOutlinerViewRow>& InRow) -> TSharedRef<SWidget>
 	{
-		return SNew(SHorizontalBox)
-			+ SHorizontalBox::Slot()
-			.AutoWidth()
-			[
-				SNew(SSeparator)
-				.Orientation(Orient_Vertical)
-				.Thickness(1.0f)
-				.SeparatorImage(FAppStyle::Get().GetBrush("Separator"))
-			];
+		return SNew(SBorder).BorderImage(FAppStyle::Get().GetBrush("Separator"));
 	};
 
 	// Left padding
