Fix rendering issue where color property tracks in Materials were causing key areas to get culled due to mismatching outliner and trackview heights.

diff --git a/Engine/Source/Editor/SequencerCore/Private/MVVM/Views/SOutlinerView.cpp b/Engine/Source/Editor/SequencerCore/Private/MVVM/Views/SOutlinerView.cpp
--- a/Engine/Source/Editor/SequencerCore/Private/MVVM/Views/SOutlinerView.cpp
+++ b/Engine/Source/Editor/SequencerCore/Private/MVVM/Views/SOutlinerView.cpp
@@ -213,6 +213,17 @@ FReply SOutlinerViewRow::OnAcceptDrop(const FDragDropEvent& DragDropEvent, EItem
 	return FReply::Unhandled();
 }
 
+FVector2D SOutlinerViewRow::ComputeDesiredSize(float LayoutScaleMultiplier) const
+{
+	FVector2D ReturnDesiredSize = SBorder::ComputeDesiredSize(LayoutScaleMultiplier);
+	if (TrackLane && TrackLane->GetOutlinerItem())
+	{
+		// Ensure our height properly matches the height of the outliner item.
+		ReturnDesiredSize.Y = TrackLane->GetOutlinerItem()->GetOutlinerSizing().GetTotalHeight();
+	}
+	return ReturnDesiredSize;
+}
+
 TViewModelPtr<IOutlinerExtension> SOutlinerViewRow::GetDataModel() const
 {
 	return WeakModel.Pin();
diff --git a/Engine/Source/Editor/SequencerCore/Public/MVVM/Views/SOutlinerView.h b/Engine/Source/Editor/SequencerCore/Public/MVVM/Views/SOutlinerView.h
--- a/Engine/Source/Editor/SequencerCore/Public/MVVM/Views/SOutlinerView.h
+++ b/Engine/Source/Editor/SequencerCore/Public/MVVM/Views/SOutlinerView.h
@@ -358,6 +358,9 @@ public:
 	/** Called to complete a drag and drop onto this drop. */
 	FReply OnAcceptDrop( const FDragDropEvent& DragDropEvent, EItemDropZone ItemDropZone, TWeakViewModelPtr<IOutlinerExtension> InDataModel);
 
+protected:
+	FVector2D ComputeDesiredSize(float LayoutScaleMultiplier) const override;
+
 private:
 
 	/**
