Sequencer: OnHierarchyUpdated can be called while the widget is being destroyed, so just mark bWidgetsDirty and not RecreateWidgets() immediately. bWidgetsDirty will call RecreateWidgets() on Tick() if necessary.

diff --git a/Engine/Source/Editor/SequencerCore/Private/MVVM/Views/STrackLane.cpp b/Engine/Source/Editor/SequencerCore/Private/MVVM/Views/STrackLane.cpp
--- a/Engine/Source/Editor/SequencerCore/Private/MVVM/Views/STrackLane.cpp
+++ b/Engine/Source/Editor/SequencerCore/Private/MVVM/Views/STrackLane.cpp
@@ -136,7 +136,6 @@ bool STrackLane::IsPinned() const
 void STrackLane::OnHierarchyUpdated()
 {
 	bWidgetsDirty = true;
-	RecreateWidgets();
 }
 
 TSharedPtr<ITrackLaneWidget> STrackLane::FindWidgetForModel(const FWeakViewModelPtr& InModel) const
