Sequencer: stable sort sequencer outliner children to avoid weird reordering issues

diff --git a/Engine/Source/Editor/SequencerCore/Private/MVVM/Extensions/ISortableExtension.cpp b/Engine/Source/Editor/SequencerCore/Private/MVVM/Extensions/ISortableExtension.cpp
--- a/Engine/Source/Editor/SequencerCore/Private/MVVM/Extensions/ISortableExtension.cpp
+++ b/Engine/Source/Editor/SequencerCore/Private/MVVM/Extensions/ISortableExtension.cpp
@@ -2,7 +2,7 @@
 
 #include "MVVM/Extensions/ISortableExtension.h"
 
-#include "Algo/Sort.h"
+#include "Algo/StableSort.h"
 #include "Containers/Array.h"
 #include "HAL/PlatformCrt.h"
 #include "Internationalization/Text.h"
@@ -102,10 +102,10 @@ void ISortableExtension::SortChildren(FViewModelChildren& Children, ESortingMode
 	{
 		case ESortingMode::PriorityFirst:
 		default:
-			Algo::Sort(ChildrenArray, FModelAndSortingKey::ComparePriorityFirst);
+			Algo::StableSort(ChildrenArray, FModelAndSortingKey::ComparePriorityFirst);
 			break;
 		case ESortingMode::CustomOrderFirst:
-			Algo::Sort(ChildrenArray, FModelAndSortingKey::CompareCustomOrderFirst);
+			Algo::StableSort(ChildrenArray, FModelAndSortingKey::CompareCustomOrderFirst);
 			break;
 	}
 
