Sequencer: Additional Tree View UX - Null pointer check

diff --git a/Engine/Source/Editor/SequencerCore/Private/MVVM/Views/SOutlinerView.cpp b/Engine/Source/Editor/SequencerCore/Private/MVVM/Views/SOutlinerView.cpp
--- a/Engine/Source/Editor/SequencerCore/Private/MVVM/Views/SOutlinerView.cpp
+++ b/Engine/Source/Editor/SequencerCore/Private/MVVM/Views/SOutlinerView.cpp
@@ -895,9 +895,12 @@ void SOutlinerView::Private_UpdateParentHighlights()
 	for (TWeakViewModelPtr<IOutlinerExtension> WeakSelectedItem : GetSelectedItems())
 	{
 		TViewModelPtr<IOutlinerExtension> SelectedItem = WeakSelectedItem.Pin();
-		for (TViewModelPtr<IOutlinerExtension> Parent : SelectedItem.AsModel()->GetAncestorsOfType<IOutlinerExtension>())
+		if (SelectedItem.IsValid())
 		{
-			Private_SetItemHighlighted(Parent, true);
+			for (TViewModelPtr<IOutlinerExtension> Parent : SelectedItem.AsModel()->GetAncestorsOfType<IOutlinerExtension>())
+			{
+				Private_SetItemHighlighted(Parent, true);
+			}
 		}
 	}
 }
