Sequencer: Track Area UX - selection for track lane is now tinted rather than a border

diff --git a/Engine/Source/Editor/Sequencer/Private/MVVM/Views/SChannelView.cpp b/Engine/Source/Editor/Sequencer/Private/MVVM/Views/SChannelView.cpp
--- a/Engine/Source/Editor/Sequencer/Private/MVVM/Views/SChannelView.cpp
+++ b/Engine/Source/Editor/Sequencer/Private/MVVM/Views/SChannelView.cpp
@@ -607,15 +607,15 @@ int32 SChannelView::DrawLane(const FPaintArgs& Args, const FGeometry& AllottedGe
 	if (Selection.Outliner.IsSelected(OutlinerItem))
 	{
 		static const FName SelectionColorName("SelectionColor");
-		static const FName SelectedTrackTintBrushName("Sequencer.Section.SelectedTrackTint");
+		static const FName HighlightBrushName("Sequencer.AnimationOutliner.DefaultBorder");
 
 		FSlateDrawElement::MakeBox(
 			OutDrawElements,
 			LayerId++,
 			AllottedGeometry.ToPaintGeometry(),
-			FAppStyle::GetBrush(SelectedTrackTintBrushName),
+			FAppStyle::GetBrush(HighlightBrushName),
 			DrawEffects,
-			FAppStyle::GetSlateColor(SelectionColorName).GetColor(InWidgetStyle)
+			FAppStyle::GetSlateColor(SelectionColorName).GetColor(InWidgetStyle).CopyWithNewOpacity(0.2f)
 		);
 	}
 
diff --git a/Engine/Source/Editor/SequencerCore/Private/MVVM/Views/STrackLane.cpp b/Engine/Source/Editor/SequencerCore/Private/MVVM/Views/STrackLane.cpp
--- a/Engine/Source/Editor/SequencerCore/Private/MVVM/Views/STrackLane.cpp
+++ b/Engine/Source/Editor/SequencerCore/Private/MVVM/Views/STrackLane.cpp
@@ -234,8 +234,6 @@ void STrackLane::OnArrangeChildren( const FGeometry& AllottedGeometry, FArranged
 
 int32 STrackLane::OnPaint(const FPaintArgs& Args, const FGeometry& AllottedGeometry, const FSlateRect& MyCullingRect, FSlateWindowElementList& OutDrawElements, int32 LayerId, const FWidgetStyle& InWidgetStyle, bool bParentEnabled) const
 {
-	LayerId = PaintLaneBackground(AllottedGeometry, MyCullingRect, OutDrawElements, LayerId, InWidgetStyle);
-
 	FArrangedChildren ArrangedChildren(EVisibility::Visible);
 	ArrangeChildren(AllottedGeometry, ArrangedChildren);
 
@@ -246,10 +244,12 @@ int32 STrackLane::OnPaint(const FPaintArgs& Args, const FGeometry& AllottedGeome
 		LayerId = CurWidget.Widget->Paint( Args.WithNewParent(this), CurWidget.Geometry, ChildClipRect, OutDrawElements, LayerId, InWidgetStyle, ShouldBeEnabled( bParentEnabled ) );
 	}
 
+	LayerId = PaintLaneForeground(AllottedGeometry, MyCullingRect, OutDrawElements, LayerId, InWidgetStyle);
+
 	return LayerId + 1;
 }
 
-int32 STrackLane::PaintLaneBackground(const FGeometry& AllottedGeometry, const FSlateRect& MyCullingRect, FSlateWindowElementList& OutDrawElements, int32 LayerId, const FWidgetStyle& InWidgetStyle) const
+int32 STrackLane::PaintLaneForeground(const FGeometry& AllottedGeometry, const FSlateRect& MyCullingRect, FSlateWindowElementList& OutDrawElements, int32 LayerId, const FWidgetStyle& InWidgetStyle) const
 {
 	static const FName BorderName("Sequencer.AnimationOutliner.DefaultBorder");
 	static const FName SelectionColorName("SelectionColor");
@@ -261,8 +261,6 @@ int32 STrackLane::PaintLaneBackground(const FGeometry& AllottedGeometry, const F
 		return LayerId;
 	}
 
-	const int32 PaintLayerId = TrackAreaView->GetPaintLayers().LaneBackgrounds;
-
 	TViewModelPtr<IHoveredExtension> Hoverable = OutlinerItem.ImplicitCast();
 
 	float TotalNodeHeight = OutlinerItem->GetOutlinerSizing().GetTotalHeight();
@@ -276,14 +274,14 @@ int32 STrackLane::PaintLaneBackground(const FGeometry& AllottedGeometry, const F
 
 		FSlateDrawElement::MakeBox(
 			OutDrawElements,
-			PaintLayerId,
+			LayerId++,
 			AllottedGeometry.ToPaintGeometry(
 				FVector2f(AllottedGeometry.GetLocalSize().X, TotalNodeHeight),
 				FSlateLayoutTransform()
 			),
 			FAppStyle::GetBrush(BorderName),
 			ESlateDrawEffect::None,
-			SelectionColor
+			SelectionColor.CopyWithNewOpacity(0.2f)
 		);
 	}
 
@@ -307,7 +305,7 @@ int32 STrackLane::PaintLaneBackground(const FGeometry& AllottedGeometry, const F
 		{
 			FSlateDrawElement::MakeBox(
 				OutDrawElements,
-				PaintLayerId,
+				LayerId++,
 				AllottedGeometry.ToPaintGeometry(
 					FVector2f(AllottedGeometry.GetLocalSize().X, TotalNodeHeight),
 					FSlateLayoutTransform()
diff --git a/Engine/Source/Editor/SequencerCore/Public/MVVM/Views/STrackLane.h b/Engine/Source/Editor/SequencerCore/Public/MVVM/Views/STrackLane.h
--- a/Engine/Source/Editor/SequencerCore/Public/MVVM/Views/STrackLane.h
+++ b/Engine/Source/Editor/SequencerCore/Public/MVVM/Views/STrackLane.h
@@ -83,7 +83,7 @@ public:
 
 protected:
 
-	int32 PaintLaneBackground(const FGeometry& AllottedGeometry, const FSlateRect& MyCullingRect, FSlateWindowElementList& OutDrawElements, int32 LayerId, const FWidgetStyle& InWidgetStyle) const;
+	int32 PaintLaneForeground(const FGeometry& AllottedGeometry, const FSlateRect& MyCullingRect, FSlateWindowElementList& OutDrawElements, int32 LayerId, const FWidgetStyle& InWidgetStyle) const;
 
 	void OnHierarchyUpdated();
 	void RecreateWidgets();
