Sequencer: Outliner UX improvements

diff --git a/Engine/Source/Editor/SequencerCore/Private/MVVM/Views/SOutlinerView.cpp b/Engine/Source/Editor/SequencerCore/Private/MVVM/Views/SOutlinerView.cpp
--- a/Engine/Source/Editor/SequencerCore/Private/MVVM/Views/SOutlinerView.cpp
+++ b/Engine/Source/Editor/SequencerCore/Private/MVVM/Views/SOutlinerView.cpp
@@ -1,8 +1,10 @@
 // Copyright Epic Games, Inc. All Rights Reserved.
 
 #include "MVVM/Views/SOutlinerView.h"
+#include "MVVM/Views/SOutlinerViewRow.h"
 
 #include "Algo/BinarySearch.h"
+#include "Algo/Find.h"
 #include "Containers/ArrayView.h"
 #include "Framework/Layout/Overscroll.h"
 #include "Framework/SlateDelegates.h"
@@ -12,7 +14,6 @@
 #include "InputCoreTypes.h"
 #include "Internationalization/Internationalization.h"
 #include "Internationalization/Text.h"
-#include "ISequencerOutlinerColumn.h"
 #include "Layout/ChildrenBase.h"
 #include "Layout/Geometry.h"
 #include "Layout/Margin.h"
@@ -28,9 +29,10 @@
 #include "MVVM/ViewModels/OutlinerViewModelDragDropOp.h"
 #include "MVVM/ViewModels/ViewModel.h"
 #include "MVVM/ViewModels/ViewModelIterators.h"
+#include "MVVM/ViewModels/OutlinerColumns/IOutlinerColumn.h"
+#include "MVVM/ViewModels/OutlinerColumns/OutlinerColumnTypesPrivate.h"
 #include "MVVM/Selection/SequencerCoreSelection.h"
 #include "MVVM/Selection/SequencerOutlinerSelection.h"
-#include "MVVM/Views/SOutlinerColumnBorder.h"
 #include "MVVM/Views/STrackAreaView.h"
 #include "MVVM/Views/STrackLane.h"
 #include "MVVM/Views/TreeViewTraits.h"
@@ -53,213 +55,16 @@
 #include "Widgets/SNullWidget.h"
 #include "Widgets/Views/SHeaderRow.h"
 
-class FPaintArgs;
-class FSlateRect;
-class ITableRow;
-class SWidget;
-namespace UE::Sequencer { class FEditorViewModel; }
-namespace UE::Sequencer { class FOutlinerSpacer; }
-
-namespace UE
-{
-namespace Sequencer
-{
-
-const FName SOutlinerView::TrackNameColumn("TrackArea");
-
-void SOutlinerViewRow::Construct(const FArguments& InArgs, const TSharedRef<STableViewBase>& OwnerTableView, TWeakViewModelPtr<IOutlinerExtension> InWeakModel)
-{
-	WeakModel = InWeakModel;
-
-	OnDetectDrag = InArgs._OnDetectDrag;
-	OnGenerateWidgetForColumn = InArgs._OnGenerateWidgetForColumn;
-
-	SMultiColumnTableRow::Construct(
-		SMultiColumnTableRow::FArguments()
-			.Style(&FAppStyle::Get().GetWidgetStyle<FTableRowStyle>("Sequencer.TableView.Row"))
-			.OnDragDetected(this, &SOutlinerViewRow::OnDragDetected)
-			.OnCanAcceptDrop(this, &SOutlinerViewRow::OnCanAcceptDrop)
-			.OnAcceptDrop(this, &SOutlinerViewRow::OnAcceptDrop)
-			.ShowSelection(IsSelectable())
-			.Padding(FMargin(0.f)),
-		OwnerTableView);
-}
-
-SOutlinerViewRow::~SOutlinerViewRow()
-{
-}
-
-TSharedRef<SWidget> SOutlinerViewRow::GenerateWidgetForColumn(const FName& ColumnId)
-{
-	TViewModelPtr<IOutlinerExtension> DataModel = WeakModel.Pin();
-	return DataModel ? OnGenerateWidgetForColumn.Execute(DataModel, ColumnId, SharedThis(this)) : SNullWidget::NullWidget;
-}
-
-void SOutlinerViewRow::Tick(const FGeometry& AllottedGeometry, const double InCurrentTime, const float InDeltaTime)
-{
-	StaticCastSharedPtr<SOutlinerView>(OwnerTablePtr.Pin())->ReportChildRowGeometry(WeakModel.Pin(), AllottedGeometry);
-}
-
-bool SOutlinerViewRow::IsSelectable() const
-{
-	TSharedPtr<ISelectableExtension> Selectable = WeakModel.ImplicitPin();
-	return !Selectable || Selectable->IsSelectable() != ESelectionIntent::Never;
-}
-
-FReply SOutlinerViewRow::OnDragDetected( const FGeometry& InGeometry, const FPointerEvent& InPointerEvent )
-{
-	if (OnDetectDrag.IsBound())
-	{
-		return OnDetectDrag.Execute(InGeometry, InPointerEvent, SharedThis(this));
-	}
-	return FReply::Unhandled();
-}
-
-int32 SOutlinerViewRow::OnPaintDropIndicator(EItemDropZone InItemDropZone, const FPaintArgs& Args, const FGeometry& AllottedGeometry, const FSlateRect& MyCullingRect, FSlateWindowElementList& OutDrawElements, int32 LayerId, const FWidgetStyle& InWidgetStyle, bool bParentEnabled) const
+namespace UE::Sequencer
 {
-	// Draw feedback for user dropping an item above, below, or onto a row.
-	const FSlateBrush* DropIndicatorBrush = GetDropIndicatorBrush(InItemDropZone);
 
-	// Offset by the indentation amount
-	static float OffsetX = 10.0f;
-	FVector2D Offset(OffsetX * GetIndentLevel(), 0.f);
-	FSlateDrawElement::MakeBox
-	(
-		OutDrawElements,
-		LayerId++,
-		AllottedGeometry.ToPaintGeometry(FVector2D(AllottedGeometry.GetLocalSize() - Offset), FSlateLayoutTransform(Offset)),
-		DropIndicatorBrush,
-		ESlateDrawEffect::None,
-		DropIndicatorBrush->GetTint(InWidgetStyle) * InWidgetStyle.GetColorAndOpacityTint()
-	);
-
-	return LayerId;
-}
-
-TOptional<EItemDropZone> SOutlinerViewRow::OnCanAcceptDrop(const FDragDropEvent& DragDropEvent, EItemDropZone InItemDropZone, TWeakViewModelPtr<IOutlinerExtension> InDataModel)
-{
-	TViewModelPtr<IOutlinerExtension> ThisModel = InDataModel.Pin();
-	if (!ThisModel)
-	{
-		return TOptional<EItemDropZone>();
-	}
-	else if (InItemDropZone == EItemDropZone::BelowItem && ThisModel->IsExpanded())
-	{
-		// Cannot drop immediately below items that are expanded since
-		// it looks like you are dropping into the item, but the object
-		// will end up below all its expanded children.
-		return TOptional<EItemDropZone>();
-	}
-
-	// The model we are indirectly dropping into (either this model, or its parent)
-	FViewModelPtr DropModel = ThisModel;
-	// The model we are directly interacting with - only set when attaching before/after
-	FViewModelPtr TargetModel;
-
-	if (!DropModel)
-	{
-		return TOptional<EItemDropZone>();
-	}
-
-	// When dropping above or below an item, we always forward to its parent
-	if (InItemDropZone == EItemDropZone::AboveItem || InItemDropZone == EItemDropZone::BelowItem)
-	{
-		TargetModel = DropModel;
-		DropModel = DropModel->CastParent<IOutlinerDropTargetOutlinerExtension>();
-	}
-
-	TViewModelPtr<IOutlinerDropTargetOutlinerExtension> DropTarget = DropModel.ImplicitCast();
-	if (!DropTarget)
-	{
-		return TOptional<EItemDropZone>();
-	}
-
-	return DropTarget->CanAcceptDrop(TargetModel, DragDropEvent, InItemDropZone);
-}
-
-FReply SOutlinerViewRow::OnAcceptDrop(const FDragDropEvent& DragDropEvent, EItemDropZone InItemDropZone, TWeakViewModelPtr<IOutlinerExtension> InDataModel)
-{
-	TViewModelPtr<IOutlinerExtension> ThisModel = InDataModel.Pin();
-	if (!ThisModel)
-	{
-		return FReply::Unhandled();
-	}
-	else if (InItemDropZone == EItemDropZone::BelowItem && ThisModel->IsExpanded())
-	{
-		// Cannot drop immediately below items that are expanded since
-		// it looks like you are dropping into the item, but the object
-		// will end up below all its expanded children.
-		return FReply::Unhandled();
-	}
-
-	// The model we are indirectly dropping into (either this model, or its parent)
-	FViewModelPtr DropModel = ThisModel;
-	// The model we are directly interacting with - only set when attaching before/after
-	FViewModelPtr TargetModel;
-
-	// When dropping above or below an item, we always forward to its parent
-	if (InItemDropZone == EItemDropZone::AboveItem || InItemDropZone == EItemDropZone::BelowItem)
-	{
-		TargetModel = DropModel;
-		DropModel = DropModel->CastParent<IOutlinerDropTargetOutlinerExtension>();
-	}
-
-	TViewModelPtr<IOutlinerDropTargetOutlinerExtension> DropTarget = DropModel.ImplicitCast();
-	if (DropTarget)
-	{
-		DropTarget->PerformDrop(TargetModel, DragDropEvent, InItemDropZone);
-		return FReply::Handled();
-	}
-
-	return FReply::Unhandled();
-}
-
-FVector2D SOutlinerViewRow::ComputeDesiredSize(float LayoutScaleMultiplier) const
-{
-	FVector2D ReturnDesiredSize = SBorder::ComputeDesiredSize(LayoutScaleMultiplier);
-	if (TrackLane && TrackLane->GetOutlinerItem())
-	{
-		// Ensure our height properly matches the height of the outliner item.
-		ReturnDesiredSize.Y = TrackLane->GetOutlinerItem()->GetOutlinerSizing().GetTotalHeight();
-	}
-	return ReturnDesiredSize;
-}
-
-TViewModelPtr<IOutlinerExtension> SOutlinerViewRow::GetDataModel() const
-{
-	return WeakModel.Pin();
-}
-
-TSharedPtr<STrackLane> SOutlinerViewRow::GetTrackLane(bool bOnlyOwnTrackLane) const
-{
-	if (!bOnlyOwnTrackLane)
-	{
-		// Return the track lane, regardless of it being created by our own outliner item, or
-		// being a reference to a parent outliner item's track lane.
-		return TrackLane;
-	}
-	if (TrackLane && TrackLane->GetOutlinerItem() == GetDataModel())
-	{
-		// Return the track lane only if it was created by our own outliner item.
-		return TrackLane;
-	}
-	return nullptr;
-}
-
-void SOutlinerViewRow::SetTrackLane(const TSharedPtr<STrackLane>& InTrackLane)
-{
-	TrackLane = InTrackLane;
-}
+SOutlinerView::SOutlinerView()
+{}
 
-const FSlateBrush* SOutlinerViewRow::GetBorder() const 
+SOutlinerView::~SOutlinerView()
 {
-	TSharedPtr<IOutlinerExtension> OutlinerItem = WeakModel.ImplicitPin();
-	if (OutlinerItem && !OutlinerItem->HasBackground())
-	{
-		return nullptr;
-	}
-
-	return SMultiColumnTableRow::GetBorder();
+	TrackArea.Reset();
+	PinnedTreeViews.Empty();
 }
 
 void SOutlinerView::Construct(const FArguments& InArgs, TWeakPtr<FOutlinerViewModel> InWeakOutliner, const TSharedRef<STrackAreaView>& InTrackArea)
@@ -283,7 +88,11 @@ void SOutlinerView::Construct(const FArguments& InArgs, TWeakPtr<FOutlinerViewMo
 		}
 	}
 
-	HeaderRow = SNew(SHeaderRow).Visibility(EVisibility::Collapsed);
+	ColumnMetaData = MakeShared<FOutlinerHeaderRowWidgetMetaData>();
+
+	HeaderRow = SNew(SHeaderRow)
+	.Visibility(EVisibility::Collapsed)
+	.AddMetaData(ColumnMetaData.ToSharedRef());
 
 	UpdateOutlinerColumns();
 
@@ -311,15 +120,6 @@ void SOutlinerView::Construct(const FArguments& InArgs, TWeakPtr<FOutlinerViewMo
 	UpdateViewSelectionFromModel();
 }
 
-SOutlinerView::SOutlinerView()
-{}
-
-SOutlinerView::~SOutlinerView()
-{
-	TrackArea.Reset();
-	PinnedTreeViews.Empty();
-}
-
 void SOutlinerView::Tick(const FGeometry& AllottedGeometry, const double InCurrentTime, const float InDeltaTime)
 {
 	const bool bWasPendingRefresh = IsPendingRefresh();
@@ -482,155 +282,146 @@ float SOutlinerView::VirtualToPhysical(float InVirtual) const
 	return InVirtual;
 }
 
-void SOutlinerView::AddOutlinerColumnCallback(TSharedPtr<ISequencerOutlinerColumn> InColumn, TSharedPtr<FEditorViewModel> InEditorViewModel)
+int32 SOutlinerView::CreateOutlinerColumnsForGroup(int32 ColumnIndex, EOutlinerColumnGroup Group)
 {
-	auto GenerateToolColumn = [=](const TWeakViewModelPtr<IOutlinerExtension>& InWeakModel, const TSharedRef<SOutlinerViewRow>& InRow) -> TSharedRef<SWidget>
-	{
-		if (TViewModelPtr<IOutlinerExtension> OutlinerExtension = InWeakModel.ImplicitPin())
-		{
-			if (!OutlinerExtension.AsModel()->IsA<FOutlinerSpacer>())
-			{
-				const FCreateOutlinerColumnParams ColumnParams = FCreateOutlinerColumnParams(OutlinerExtension, InEditorViewModel);
-				const bool bIsCompatible = InColumn->IsItemCompatibleWithColumn(ColumnParams);
+	const int32 NumColumns = OutlinerColumns.Num();
 
-				TSharedRef<SWidget> InnerContent = bIsCompatible ? InColumn->CreateColumnWidget(InColumn, ColumnParams) : SNullWidget::NullWidget;
+	int32 NumAdded = 0;
+	for ( ; ColumnIndex < NumColumns; ++ColumnIndex, ++NumAdded)
+	{
+		const TSharedPtr<IOutlinerColumn>& Column = OutlinerColumns[ColumnIndex];
 
-				return SNew(SOutlinerColumnBorder, ColumnParams, bIsCompatible)
-					[
-						InnerContent
-					];
-			}
-			else
-			{
-				return SNew(SBox).HeightOverride(10.f);
-			}
+		FOutlinerColumnPosition Position = Column->GetPosition();
+		if (Position.Group != Group)
+		{
+			// OutlinerColumns is sorted so if we encountered the wrong group we should stop immediately
+			break;
 		}
 
-		ensureMsgf(false, TEXT("Attempting to create an outliner column for a view model that is either dead, or not an outliner item."));
-		return SNew(SBox).HeightOverride(10.f);
-	};
+		FName                 ColumnName = Column->GetColumnName();
+		FOutlinerColumnLayout Layout     = Column->GetLayout();
 
-	Columns.Add(InColumn->GetColumnName(), FOutlinerViewColumn(GenerateToolColumn, 22.f, true));
-}
+		// Add the meta-data required for our custom row panel
+		ColumnMetaData->Columns.Add(Layout);
 
-void SOutlinerView::UpdateTrackGutterColumns(TSharedPtr<FEditorViewModel> InEditorViewModel)
-{
-	// Padding columns callback, used to provide apropriately colored padding to the left and right of the track gutter
-	auto GeneratePadding = [=](const TWeakViewModelPtr<IOutlinerExtension>& InWeakModel, const TSharedRef<SOutlinerViewRow>& InRow) -> TSharedRef<SWidget>
-	{
-		if (TViewModelPtr<IOutlinerExtension> OutlinerExtension = InWeakModel.ImplicitPin())
+		// Keep track of how to create widgets for this column
+		ColumnGenerators.Add(ColumnName, [Column](const FCreateOutlinerColumnParams& Params, const TSharedRef<SOutlinerViewRow>& Row){
+			return Column->IsItemCompatibleWithColumn(Params)
+				? Column->CreateColumnWidget(Params, Row)
+				: nullptr;
+		});
+
+		// Add the column itself to the header row
+		if (Layout.SizeMode == EOutlinerColumnSizeMode::Fixed)
 		{
-			if (!OutlinerExtension.AsModel()->IsA<FOutlinerSpacer>())
-			{
-				const FCreateOutlinerColumnParams ColumnParams = FCreateOutlinerColumnParams(OutlinerExtension, InEditorViewModel);
-				return SNew(SOutlinerColumnBorder, ColumnParams, false);
-			}
-			else
-			{
-				return SNew(SBox).HeightOverride(10.f);
-			}
+			HeaderRow->AddColumn(
+				SHeaderRow::Column(ColumnName)
+				.FixedWidth(Layout.Width)
+			);
+		}
+		else
+		{
+			HeaderRow->AddColumn(
+				SHeaderRow::Column(ColumnName)
+				.FillWidth(Layout.Width)
+			);
 		}
-		ensureMsgf(false, TEXT("Attempting to create an outliner column for a view model that is either dead, or not an outliner item."));
-		return SNew(SBox).HeightOverride(10.f);
-	};
-
-	// Spacer column callback, generates a line to separate the track gutter from the outliner
-	auto GenerateSpacer = [=](const TWeakViewModelPtr<IOutlinerExtension>& InWeakModel, const TSharedRef<SOutlinerViewRow>& InRow) -> TSharedRef<SWidget>
-	{
-		return SNew(SBorder).BorderImage(FAppStyle::Get().GetBrush("Separator"));
-	};
-
-	// Left padding
-	Columns.Add(FName(TEXT("OutlinerColumnLeftPadding")), FOutlinerViewColumn(GeneratePadding, 2.f, true));
-
-	// Create Visible Outliner Columns in the Track Gutter (Pin/Mute/Solo...)
-	for (TSharedPtr<ISequencerOutlinerColumn> Column : OutlinerColumns)
-	{
-		AddOutlinerColumnCallback(Column, InEditorViewModel);
 	}
 
-	// Right padding
-	Columns.Add(FName(TEXT("OutlinerColumnRightPadding")), FOutlinerViewColumn(GeneratePadding, 2.f, true));
+	return NumAdded;
+}
 
-	// Spacer between Track Gutter and Outliner View
-	Columns.Add(FName(TEXT("OutlinerColumnSpacer")), FOutlinerViewColumn(GenerateSpacer, 1.f, true));
+void SOutlinerView::InsertSeparatorColumn(int32 InsertIndex, int32 SeparatorID)
+{
+	static const FName NAME_Separator("Separator");
+
+	FOutlinerColumnLayout SeparatorLayout(
+		1.f, /* Width */
+		FMargin(0.f),
+		HAlign_Fill,
+		VAlign_Fill,
+		EOutlinerColumnSizeMode::Fixed,
+		EOutlinerColumnFlags::Hidden
+	);
+	FName SeparatorName(NAME_Separator, SeparatorID);
+
+	// Add a 1px separator column
+	ColumnMetaData->Columns.Insert(SeparatorLayout, InsertIndex);
+	HeaderRow->InsertColumn(
+		SHeaderRow::Column(SeparatorName)
+		.FixedWidth(1.f),
+		InsertIndex
+	);
+	ColumnGenerators.Add(SeparatorName,
+		[](const FCreateOutlinerColumnParams&, const TSharedRef<SOutlinerViewRow>&)->TSharedPtr<SWidget>
+		{
+			return SNew(SImage)
+			.Image(FAppStyle::Get().GetBrush("Sequencer.Outliner.Separator"));
+		}
+	);
 }
 
 void SOutlinerView::UpdateOutlinerColumns()
 {
+	// Sort the columns by position
+	Algo::SortBy(OutlinerColumns, &IOutlinerColumn::GetPosition);
+
 	// Clear columns to ensure consistent order when building UI from Map
 	HeaderRow->ClearColumns();
-	Columns.Empty();
-	Columns.Compact();
+	ColumnMetaData->Columns.Empty();
+	ColumnGenerators.Empty();
 
-	TSharedPtr<FEditorViewModel> EditorViewModel = WeakOutliner.Pin()->GetEditor();
+	// ----------------------------------------------------------------------------------------------------------
+	// Populate columns
+	const int32 NumLeftGutter  = CreateOutlinerColumnsForGroup(0,                         EOutlinerColumnGroup::LeftGutter);
+	const int32 NumCenter      = CreateOutlinerColumnsForGroup(NumLeftGutter,             EOutlinerColumnGroup::Center);
+	const int32 NumRightGutter = CreateOutlinerColumnsForGroup(NumLeftGutter + NumCenter, EOutlinerColumnGroup::RightGutter);
 
-	if (OutlinerColumns.Num() > 0)
+	// ----------------------------------------------------------------------------------------------------------
+	// Add some padding to the leading and trailing edge of the first and last columns in each group respectively
+	//    This is implemented this way because columns can be turned on and off dynamically, but we must
+	//    always have a consistent padding within the group. Separators dynamically appear based on the presence of
+	//    each group so we can't put padding on those
+	if (NumLeftGutter > 0)
 	{
-		UpdateTrackGutterColumns(EditorViewModel);
+		ColumnMetaData->Columns[0              ].CellPadding.Left  += 4.f;
+		ColumnMetaData->Columns[NumLeftGutter-1].CellPadding.Right += 4.f;
 	}
-
-	// Create Column for Outliner View (Track name / Keying)
-	auto GenerateOutliner = [=](const TWeakViewModelPtr<IOutlinerExtension>& InWeakModel, const TSharedRef<SOutlinerViewRow>& InRow) -> TSharedRef<SWidget>
+	if (NumCenter > 0)
 	{
-		if (TSharedPtr<IOutlinerExtension> OutlinerItem = InWeakModel.ImplicitPin())
-		{
-			return OutlinerItem->CreateOutlinerView(FCreateOutlinerViewParams{ InRow, EditorViewModel });
-		}
+		ColumnMetaData->Columns[NumLeftGutter            ].CellPadding.Left  += 4.f;
+		ColumnMetaData->Columns[NumLeftGutter+NumCenter-1].CellPadding.Right += 4.f;
+	}
 
-		ensureMsgf(false, TEXT("Attempting to create an outliner widget for a view model that is either dead, or not an outliner item."));
-		return SNew(SBox).HeightOverride(10.f);
-	};
+	// No additional padding on the right gutter intentionally
 
-	Columns.Add("Outliner", FOutlinerViewColumn(GenerateOutliner, 1.f));
+	// ----------------------------------------------------------------------------------------------------------
+	// Add separators between the groups.
+	//      Only add separators if they actually separate columns.
+	int32 NumSeparators = 0;
 
-	// Now populate the header row with the columns
-	for (TTuple<FName, FOutlinerViewColumn>& Pair : Columns)
+	int32 InsertIndex = NumLeftGutter;
+	if (InsertIndex < ColumnMetaData->Columns.Num())
 	{
-		if (Pair.Key != TrackNameColumn)
+		InsertSeparatorColumn(InsertIndex++, NumSeparators++);
+		if (NumCenter > 0)
 		{
-			if (Pair.Value.bIsFixedWidth)
-			{
-				TOptional<float> FixedWidth;
-				if (Pair.Value.Width.IsSet())
-				{
-					FixedWidth = Pair.Value.Width.Get();
-				}
-				HeaderRow->AddColumn(
-					SHeaderRow::Column(Pair.Key)
-					.FixedWidth(FixedWidth)
-				);
-			}
-			else
+			InsertIndex += NumCenter;
+			if (InsertIndex < ColumnMetaData->Columns.Num())
 			{
-				HeaderRow->AddColumn(
-					SHeaderRow::Column(Pair.Key)
-					.FillWidth(Pair.Value.Width)
-				);
+				InsertSeparatorColumn(InsertIndex++, NumSeparators++);
 			}
 		}
 	}
 }
 
-void SOutlinerView::UpdateTrackArea()
-{
-	// Add or remove the column
-	if (const FOutlinerViewColumn* Column = Columns.Find(TrackNameColumn))
-	{
-		HeaderRow->AddColumn(
-			SHeaderRow::Column(TrackNameColumn)
-			.FillWidth(Column->Width)
-		);
-	}
-}
-
 void SOutlinerView::AddPinnedTreeView(TSharedPtr<SOutlinerView> PinnedTreeView)
 {
 	PinnedTreeViews.Add(PinnedTreeView);
 	PinnedTreeView->SetPrimaryTreeView(SharedThis(this));
 }
 
-void SOutlinerView::SetOutlinerColumns(const TArray<TSharedPtr<ISequencerOutlinerColumn>>& InOutlinerColumns)
+void SOutlinerView::SetOutlinerColumns(const TArray<TSharedPtr<IOutlinerColumn>>& InOutlinerColumns)
 {
 	// Reset the way rows are constructed with an updated list of Outliner Columns
 	OutlinerColumns = InOutlinerColumns;
@@ -1378,6 +1169,7 @@ TSharedRef<ITableRow> SOutlinerView::OnGenerateRow(TWeakViewModelPtr<IOutlinerEx
 	TSharedRef<SOutlinerViewRow> Row =
 		SNew(SOutlinerViewRow, OwnerTable, InWeakModel)
 		.OnDetectDrag(this, &SOutlinerView::OnDragRow)
+		.OnGetColumnVisibility(this, &SOutlinerView::IsColumnVisible)
 		.OnGenerateWidgetForColumn(this, &SOutlinerView::GenerateWidgetForColumn);
 
 	if (TViewModelPtr<IOutlinerExtension> ViewModel = InWeakModel.Pin())
@@ -1389,16 +1181,36 @@ TSharedRef<ITableRow> SOutlinerView::OnGenerateRow(TWeakViewModelPtr<IOutlinerEx
 
 TSharedRef<SWidget> SOutlinerView::GenerateWidgetForColumn(TViewModelPtr<IOutlinerExtension> InDataModel, const FName& ColumnId, const TSharedRef<SOutlinerViewRow>& Row) const
 {
-	const FOutlinerViewColumn* Definition = Columns.Find(ColumnId);
+	TSharedPtr<FEditorViewModel> Editor = GetOutlinerModel()->GetEditor();
+	if (!Editor)
+	{
+		return SNullWidget::NullWidget;
+	}
+
+	// First of all, see if the model wants to create a widget
+	TSharedPtr<SWidget> ViewModelWidget = InDataModel->CreateOutlinerViewForColumn(FCreateOutlinerViewParams{ Row, Editor.ToSharedRef() }, ColumnId);
+	if (ViewModelWidget)
+	{
+		return ViewModelWidget.ToSharedRef();
+	}
 
-	if (ensureMsgf(Definition, TEXT("Invalid column name specified")))
+	// Next see if we have a column generator function for this column
+	if (const FColumnGenerator* Generator = ColumnGenerators.Find(ColumnId))
 	{
-		return Definition->Generator(InDataModel, Row);
+		FCreateOutlinerColumnParams Params{ InDataModel, Editor };
+		if (TSharedPtr<SWidget> Result = (*Generator)(Params, Row))
+		{
+			return Result.ToSharedRef();
+		}
 	}
 
 	return SNullWidget::NullWidget;
 }
 
-} // namespace Sequencer
-} // namespace UE
+bool SOutlinerView::IsColumnVisible(const FName& InName) const
+{
+	return Algo::FindBy(OutlinerColumns, InName, &IOutlinerColumn::GetColumnName) != nullptr;
+}
+
+} // namespace UE::Sequencer
 
diff --git a/Engine/Source/Editor/SequencerCore/Private/MVVM/Views/ViewUtilities.cpp b/Engine/Source/Editor/SequencerCore/Private/MVVM/Views/ViewUtilities.cpp
--- a/Engine/Source/Editor/SequencerCore/Private/MVVM/Views/ViewUtilities.cpp
+++ b/Engine/Source/Editor/SequencerCore/Private/MVVM/Views/ViewUtilities.cpp
@@ -2,63 +2,80 @@
 
 #include "MVVM/Views/ViewUtilities.h"
 
-#include "Delegates/Delegate.h"
-#include "Fonts/SlateFontInfo.h"
-#include "Internationalization/Text.h"
-#include "Layout/Margin.h"
-#include "Layout/Visibility.h"
-#include "Misc/Attribute.h"
-#include "SlotBase.h"
 #include "Styling/AppStyle.h"
-#include "Styling/CoreStyle.h"
-#include "Styling/SlateColor.h"
-#include "Templates/TypeHash.h"
-#include "Types/SlateEnums.h"
-#include "Widgets/DeclarativeSyntaxSupport.h"
-#include "Widgets/Images/SImage.h"
-#include "Widgets/Input/SComboButton.h"
-#include "Widgets/SBoxPanel.h"
-#include "Widgets/Text/STextBlock.h"
+#include "MVVM/ViewModelPtr.h"
+#include "MVVM/ViewModels/ViewModelIterators.h"
+#include "MVVM/ViewModels/EditorViewModel.h"
+#include "MVVM/Extensions/IHoveredExtension.h"
+#include "MVVM/Views/OutlinerColumns/SOutlinerColumnButton.h"
 
-class SWidget;
 
-namespace UE
+namespace UE::Sequencer
 {
-namespace Sequencer
+
+TSharedRef<SWidget> MakeAddButton(FText HoverText, const FOnClicked& HandleClicked, const FViewModelPtr& ViewModel)
 {
+	TAttribute<bool> IsHovered;
+	if (TSharedPtr<IHoveredExtension> Hoverable = ViewModel.ImplicitCast())
+	{
+		IsHovered = MakeAttributeSP(Hoverable.ToSharedRef(), &IHoveredExtension::IsHovered);
+	}
+
+
+	TAttribute<bool> IsEnabled;
+	if (TSharedPtr<FEditorViewModel> Editor = ViewModel->FindAncestorOfType<FEditorViewModel>())
+	{
+		IsEnabled = MakeAttributeSP(Editor.ToSharedRef(), &FEditorViewModel::IsEditable);
+	}
+
+	return SNew(SOutlinerColumnButton)
+		.ToolTipText(HoverText)
+		.OnClicked(HandleClicked)
+		.IsEnabled(IsEnabled)
+		.IsRowHovered(IsHovered)
+		.Image(FAppStyle::GetBrush("Sequencer.Outliner.Plus"));
+}
 
-TSharedRef<SWidget> MakeAddButton(FText HoverText, FOnGetContent MenuContent, const TAttribute<bool>& HoverState, const TAttribute<bool>& IsEnabled)
+TSharedRef<SWidget> MakeAddButton(FText HoverText, const FOnGetContent& HandleGetMenuContent, const FViewModelPtr& ViewModel)
 {
-	TSharedRef<SComboButton> ComboButton =
+	TAttribute<bool> IsHovered;
+	if (TSharedPtr<IHoveredExtension> Hoverable = ViewModel.ImplicitCast())
+	{
+		IsHovered = MakeAttributeSP(Hoverable.ToSharedRef(), &IHoveredExtension::IsHovered);
+	}
 
-		SNew(SComboButton)
-		.HasDownArrow(false)
-		.IsFocusable(true)
-		.ButtonStyle(FAppStyle::Get(), "HoverHintOnly")
-		.ForegroundColor( FSlateColor::UseForeground() )
+	TAttribute<bool> IsEnabled;
+	if (TSharedPtr<FEditorViewModel> Editor = ViewModel->FindAncestorOfType<FEditorViewModel>())
+	{
+		IsEnabled = MakeAttributeSP(Editor.ToSharedRef(), &FEditorViewModel::IsEditable);
+	}
+
+	return SNew(SOutlinerColumnButton)
+		.ToolTipText(HoverText)
+		.OnGetMenuContent(HandleGetMenuContent)
 		.IsEnabled(IsEnabled)
-		.OnGetMenuContent(MenuContent)
-		.ContentPadding(FMargin(5, 2))
-		.HAlign(HAlign_Center)
-		.VAlign(VAlign_Center)
-		.ButtonContent()
-		[
-			SNew(SHorizontalBox)
+		.IsRowHovered(IsHovered)
+		.Image(FAppStyle::GetBrush("Sequencer.Outliner.Plus"));
+}
 
-			+ SHorizontalBox::Slot()
-			.AutoWidth()
-			.VAlign(VAlign_Center)
-			.Padding(FMargin(0,0,2,0))
-			[
-				SNew(SImage)
-				.ColorAndOpacity( FSlateColor::UseForeground() )
-				.Image(FAppStyle::GetBrush("Plus"))
-				.ToolTipText(HoverText)
-			]
-		];
+TSharedRef<SWidget> MakeAddButton(FText HoverText, const FOnGetContent& MenuContent, const TAttribute<bool>& HoverState, const TAttribute<bool>& IsEnabled)
+{
+	return SNew(SOutlinerColumnButton)
+		.ToolTipText(HoverText)
+		.OnGetMenuContent(MenuContent)
+		.IsEnabled(IsEnabled)
+		.IsRowHovered(HoverState)
+		.Image(FAppStyle::GetBrush("Sequencer.Outliner.Plus"));
+}
 
-	return ComboButton;
+TSharedRef<SWidget> MakeAddButton(FText HoverText, const FOnClicked& OnClicked, const TAttribute<bool>& HoverState, const TAttribute<bool>& IsEnabled)
+{
+	return SNew(SOutlinerColumnButton)
+		.ToolTipText(HoverText)
+		.OnClicked(OnClicked)
+		.IsEnabled(IsEnabled)
+		.IsRowHovered(HoverState)
+		.Image(FAppStyle::GetBrush("Sequencer.Outliner.Plus"));
 }
 
-} // namespace Sequencer
-} // namespace UE
\ No newline at end of file
+} // namespace UE::Sequencer
\ No newline at end of file
diff --git a/Engine/Source/Editor/SequencerCore/Public/MVVM/Extensions/IOutlinerExtension.h b/Engine/Source/Editor/SequencerCore/Public/MVVM/Extensions/IOutlinerExtension.h
--- a/Engine/Source/Editor/SequencerCore/Public/MVVM/Extensions/IOutlinerExtension.h
+++ b/Engine/Source/Editor/SequencerCore/Public/MVVM/Extensions/IOutlinerExtension.h
@@ -51,6 +51,8 @@ enum class EOutlinerSizingFlags
 {
 	None                      = 0,
 	DynamicSizing             = 1 << 0,
+	IncludeSeparator          = 1 << 1,
+	CustomHeight              = 1 << 2,
 };
 ENUM_CLASS_FLAGS(EOutlinerSizingFlags);
 
@@ -81,35 +83,48 @@ struct SEQUENCERCORE_API FOutlinerSizing
 		return !(A == B);
 	}
 
+	float GetSeparatorHeight() const
+	{
+		return EnumHasAnyFlags(Flags, EOutlinerSizingFlags::IncludeSeparator)
+			? 1.f
+			: 0.f;
+	}
+
 	float GetTotalHeight() const
 	{
-		return Height + PaddingTop + PaddingBottom;
+		return Height + PaddingTop + PaddingBottom + GetSeparatorHeight();
 	}
 
 	void Accumulate(const FOutlinerSizing& Other)
 	{
-		Height        = FMath::Max(Height, Other.Height);
+		const float ThisSeparatorHeight  = GetSeparatorHeight();
+		const float OtherSeparatorHeight = Other.GetSeparatorHeight();
+
+		Height        = FMath::Max(Height + ThisSeparatorHeight, Other.Height + OtherSeparatorHeight);
 		PaddingTop    = FMath::Max(PaddingTop, Other.PaddingTop);
 		PaddingBottom = FMath::Max(PaddingBottom, Other.PaddingBottom);
 		Flags |= Other.Flags;
+
+		Height -= GetSeparatorHeight();
 	}
 
 	float Height = 0.f;
 	float PaddingTop = 0.f;
 	float PaddingBottom = 0.f;
-	EOutlinerSizingFlags Flags = EOutlinerSizingFlags::None;
-
+	EOutlinerSizingFlags Flags = EOutlinerSizingFlags::IncludeSeparator;
 };
 
 /** Facade class to allow SOutlinerViewRow to be passed around publicly */
-class SEQUENCERCORE_API ISequencerTreeViewRow : public SMultiColumnTableRow<TWeakViewModelPtr<IOutlinerExtension>>
+class SEQUENCERCORE_API ISequencerTreeViewRow : public STableRow<TWeakViewModelPtr<IOutlinerExtension>>
 {
+public:
+	virtual bool IsColumnVisible(const FName& InColumnName) const = 0;
 };
 
 /** Parameters for creating an outliner item widget. */
-struct SEQUENCERCORE_API FCreateOutlinerViewParams
+struct FCreateOutlinerViewParams
 {
-	FCreateOutlinerViewParams(const TSharedRef<ISequencerTreeViewRow> InTreeViewRow, const TSharedPtr<FEditorViewModel> InEditor)
+	FCreateOutlinerViewParams(const TSharedRef<ISequencerTreeViewRow>& InTreeViewRow, const TSharedPtr<FEditorViewModel>& InEditor)
 		: TreeViewRow(InTreeViewRow), Editor(InEditor)
 	{}
 
@@ -118,7 +133,7 @@ struct SEQUENCERCORE_API FCreateOutlinerViewParams
 };
 
 /** Parameters for building a context menu widget */
-struct SEQUENCERCORE_API FCreateOutlinerContextMenuWidgetParams
+struct FCreateOutlinerContextMenuWidgetParams
 {
 	FCreateOutlinerContextMenuWidgetParams(const TSharedPtr<FEditorViewModel> InEditor)
 		: Editor(InEditor)
@@ -165,8 +180,11 @@ public:
 	/** Toggles a bit in the outliner item's selection state */
 	virtual void ToggleSelectionState(EOutlinerSelectionState InState, bool bInValue);
 
-	/** Create the visual outliner widget for the owning model */
-	virtual TSharedRef<SWidget> CreateOutlinerView(const FCreateOutlinerViewParams& InParams) = 0;
+	/** (deprecated) Create the outliner view for the label column */
+	virtual TSharedRef<SWidget> CreateOutlinerView(const FCreateOutlinerViewParams& InParams);
+
+	/** Create the visual outliner widget for the specified column */
+	virtual TSharedPtr<SWidget> CreateOutlinerViewForColumn(const FCreateOutlinerViewParams& InParams, const FName& InColumnName);
 
 	/** Build the context menu widget for this outliner item */
 	virtual TSharedPtr<SWidget> CreateContextMenuWidget(const FCreateOutlinerContextMenuWidgetParams& InParams) = 0;
diff --git a/Engine/Source/Editor/SequencerCore/Public/MVVM/SharedViewModelData.h b/Engine/Source/Editor/SequencerCore/Public/MVVM/SharedViewModelData.h
--- a/Engine/Source/Editor/SequencerCore/Public/MVVM/SharedViewModelData.h
+++ b/Engine/Source/Editor/SequencerCore/Public/MVVM/SharedViewModelData.h
@@ -24,6 +24,9 @@ class SEQUENCERCORE_API FSharedViewModelData
 	: public FViewModel
 {
 public:
+
+	UE_SEQUENCER_DECLARE_CASTABLE(FSharedViewModelData, FViewModel);
+
 	void PreHierarchicalChange(const TSharedPtr<FViewModel>& InChangedModel);
 	void BroadcastHierarchicalChange(const TSharedPtr<FViewModel>& InChangedModel);
 
diff --git a/Engine/Source/Editor/SequencerCore/Public/MVVM/ViewModels/EditorViewModel.h b/Engine/Source/Editor/SequencerCore/Public/MVVM/ViewModels/EditorViewModel.h
--- a/Engine/Source/Editor/SequencerCore/Public/MVVM/ViewModels/EditorViewModel.h
+++ b/Engine/Source/Editor/SequencerCore/Public/MVVM/ViewModels/EditorViewModel.h
@@ -7,6 +7,7 @@
 #include "MVVM/ViewModelPtr.h"
 #include "MVVM/ViewModelTypeID.h"
 #include "MVVM/ViewModels/ViewModel.h"
+#include "MVVM/ViewModels/ViewDensity.h"
 #include "MVVM/ViewModels/ViewModelHierarchy.h"
 #include "Misc/FrameRate.h"
 #include "Templates/SharedPointer.h"
@@ -71,9 +72,21 @@ public:
 	/** Gets the scripting layer */
 	USequencerScriptingLayer* GetScriptingLayer() const;
 
+	/** Returns the current view density for this editor */
+	FViewDensityInfo GetViewDensity() const;
+
+	/** Returns the current view density for this editor */
+	void SetViewDensity(const FViewDensityInfo& InViewDensity);
+
 	/** Returns whether this editor is currently read-only */
 	virtual bool IsReadOnly() const;
 
+	/** Returns the inverse of read-only - useful for direct bindings to IsEnabled for widgets */
+	bool IsEditable() const
+	{
+		return !IsReadOnly();
+	}
+
 protected:
 
 	/** Initializes this view model before panels and root models are created */
@@ -108,7 +121,10 @@ private:
 	TStrongObjectPtr<USequencerScriptingLayer> ScriptingLayer;
 
 	/** The root view-model for the data being edited */
-	TSharedPtr<FViewModel>  RootDataModel;
+	TSharedPtr<FViewModel> RootDataModel;
+
+	/** This editor's current view density defining how condensed the elements appear */
+	FViewDensityInfo ViewDensity;
 };
 
 } // namespace UE::Sequencer
diff --git a/Engine/Source/Editor/SequencerCore/Public/MVVM/ViewModels/OutlinerSpacer.h b/Engine/Source/Editor/SequencerCore/Public/MVVM/ViewModels/OutlinerSpacer.h
--- a/Engine/Source/Editor/SequencerCore/Public/MVVM/ViewModels/OutlinerSpacer.h
+++ b/Engine/Source/Editor/SequencerCore/Public/MVVM/ViewModels/OutlinerSpacer.h
@@ -48,7 +48,6 @@ public:
 	bool HasBackground() const override;
 	FName GetIdentifier() const override;
 	FOutlinerSizing GetOutlinerSizing() const override;
-	TSharedRef<SWidget> CreateOutlinerView(const FCreateOutlinerViewParams& InParams) override;
 	TSharedPtr<SWidget> CreateContextMenuWidget(const FCreateOutlinerContextMenuWidgetParams& InParams) override;
 
 	/*~ ISortableExtension */
diff --git a/Engine/Source/Editor/SequencerCore/Public/MVVM/ViewModels/OutlinerViewModel.h b/Engine/Source/Editor/SequencerCore/Public/MVVM/ViewModels/OutlinerViewModel.h
--- a/Engine/Source/Editor/SequencerCore/Public/MVVM/ViewModels/OutlinerViewModel.h
+++ b/Engine/Source/Editor/SequencerCore/Public/MVVM/ViewModels/OutlinerViewModel.h
@@ -15,11 +15,8 @@
 
 class FDragDropOperation;
 class SWidget;
-namespace UE::Sequencer { template <typename T> struct TAutoRegisterViewModelTypeID; }
 
-namespace UE
-{
-namespace Sequencer
+namespace UE::Sequencer
 {
 
 class FEditorViewModel;
@@ -90,6 +87,12 @@ public:
 	 */
 	virtual TSharedRef<FDragDropOperation> InitiateDrag(TArray<TWeakViewModelPtr<IOutlinerExtension>>&& InDraggedModels);
 
+	/**
+	 * Request that this tree be updated
+	 */
+	virtual void RequestUpdate();
+
+
 private:
 
 	void HandleDataHierarchyChanged();
@@ -108,6 +111,5 @@ private:
 	bool bRootItemsInvalidated;
 };
 
-} // namespace Sequencer
-} // namespace UE
+} // namespace UE::Sequencer
 
diff --git a/Engine/Source/Editor/SequencerCore/Public/MVVM/Views/OutlinerColumns/SOutlinerColumnButton.h b/Engine/Source/Editor/SequencerCore/Public/MVVM/Views/OutlinerColumns/SOutlinerColumnButton.h
new file mode 100644
--- /dev/null
+++ b/Engine/Source/Editor/SequencerCore/Public/MVVM/Views/OutlinerColumns/SOutlinerColumnButton.h
@@ -0,0 +1,47 @@
+// Copyright Epic Games, Inc. All Rights Reserved.
+
+#pragma once
+
+#include "Widgets/SCompoundWidget.h"
+#include "Framework/SlateDelegates.h"
+
+class SMenuAnchor;
+
+namespace UE::Sequencer
+{
+
+class IOutlinerExtension;
+
+class SEQUENCERCORE_API SOutlinerColumnButton : public SCompoundWidget
+{
+public:
+
+	SLATE_BEGIN_ARGS(SOutlinerColumnButton) : _IsFocusable(false) {}
+
+		SLATE_ATTRIBUTE(const FSlateBrush*, Image)
+
+		SLATE_ATTRIBUTE(bool, IsRowHovered)
+
+		SLATE_ARGUMENT(bool, IsFocusable)
+
+		SLATE_EVENT(FOnClicked, OnClicked)
+
+		SLATE_EVENT(FOnGetContent, OnGetMenuContent)
+
+	SLATE_END_ARGS()
+
+	void Construct(const FArguments& InArgs);
+
+private:
+
+	FSlateColor GetColorAndOpacity() const;
+
+	FReply ToggleMenu();
+
+private:
+
+	TSharedPtr<SMenuAnchor> MenuAnchor;
+	TAttribute<bool> IsRowHovered;
+};
+
+} // namespace UE::Sequencer
\ No newline at end of file
diff --git a/Engine/Source/Editor/SequencerCore/Public/MVVM/Views/SKeyNavigationButtons.h b/Engine/Source/Editor/SequencerCore/Public/MVVM/Views/SKeyNavigationButtons.h
--- a/Engine/Source/Editor/SequencerCore/Public/MVVM/Views/SKeyNavigationButtons.h
+++ b/Engine/Source/Editor/SequencerCore/Public/MVVM/Views/SKeyNavigationButtons.h
@@ -36,8 +36,7 @@ public:
 	DECLARE_DELEGATE_OneParam(FOnGetNavigatableTimes, TArray<FFrameNumber>&)
 	DECLARE_DELEGATE_TwoParams(FOnAddKey, FFrameTime, TSharedPtr<FViewModel>)
 
-	SLATE_BEGIN_ARGS(SKeyNavigationButtons) : _IsEnabled(true) {}
-		SLATE_ATTRIBUTE(bool, IsEnabled)
+	SLATE_BEGIN_ARGS(SKeyNavigationButtons) {}
 
 		SLATE_ARGUMENT(FText, PreviousKeyToolTip)
 		SLATE_ARGUMENT(FText, NextKeyToolTip)
@@ -53,7 +52,6 @@ public:
 
 	void Construct(const FArguments& InArgs, const TSharedPtr<FViewModel>& InModel);
 
-	FLinearColor GetHoverTint() const;
 	FReply OnPreviousKeyClicked();
 	FReply OnNextKeyClicked();
 	FReply OnAddKeyClicked();
diff --git a/Engine/Source/Editor/SequencerCore/Public/MVVM/Views/SOutlinerView.h b/Engine/Source/Editor/SequencerCore/Public/MVVM/Views/SOutlinerView.h
--- a/Engine/Source/Editor/SequencerCore/Public/MVVM/Views/SOutlinerView.h
+++ b/Engine/Source/Editor/SequencerCore/Public/MVVM/Views/SOutlinerView.h
@@ -12,6 +12,7 @@
 #include "Delegates/Delegate.h"
 #include "Input/Reply.h"
 #include "MVVM/Extensions/IOutlinerExtension.h"
+#include "MVVM/ViewModels/OutlinerColumns/IOutlinerColumn.h"
 #include "MVVM/ViewModelPtr.h"
 #include "MVVM/Views/TreeViewTraits.h"
 #include "Misc/Attribute.h"
@@ -32,7 +33,6 @@ class FPaintArgs;
 class FSlateRect;
 class FSlateWindowElementList;
 class FWidgetStyle;
-class ISequencerOutlinerColumn;
 class ITableRow;
 class SHeaderRow;
 class SScrollBar;
@@ -43,43 +43,29 @@ struct FKeyEvent;
 struct FPointerEvent;
 struct FSlateBrush;
 
-namespace UE
-{
-namespace Sequencer
+namespace UE::Sequencer
 {
 
 struct FSelectionEventSuppressor;
+struct FOutlinerHeaderRowWidgetMetaData;
 
 class FViewModel;
 class FSequencerCoreSelection;
 class SOutlinerViewRow;
 class STrackAreaView;
 class STrackLane;
+class IOutlinerColumn;
+
+enum class EOutlinerColumnGroup : uint8;
 
 enum class ETreeRecursion
 {
 	Recursive, NonRecursive
 };
 
-/** Structure used to define a column in the tree view */
-struct SEQUENCERCORE_API FOutlinerViewColumn
-{
-	typedef TFunction<TSharedRef<SWidget>(TViewModelPtr<IOutlinerExtension>, const TSharedRef<SOutlinerViewRow>&)> FOnGenerate;
-
-	FOutlinerViewColumn(const FOnGenerate& InOnGenerate, const TAttribute<float>& InWidth, const bool InIsFixedWidth) : Generator(InOnGenerate), Width(InWidth), bIsFixedWidth(InIsFixedWidth) {}
-	FOutlinerViewColumn(const FOnGenerate& InOnGenerate, const TAttribute<float>& InWidth) : Generator(InOnGenerate), Width(InWidth), bIsFixedWidth(false) {}
-	FOutlinerViewColumn(FOnGenerate&& InOnGenerate, const TAttribute<float>& InWidth) : Generator(MoveTemp(InOnGenerate)), Width(InWidth), bIsFixedWidth(false) {}
-
-	/** Function used to generate a cell for this column */
-	FOnGenerate Generator;
-	/** Attribute specifying the width of this column */
-	TAttribute<float> Width;
-	/** Attribute specifying whether or not to use the Width in percentage or fixed-width values */
-	bool bIsFixedWidth = false;
-};
 
 /** The tree view used in the sequencer */
-class SEQUENCERCORE_API SOutlinerView 
+class SEQUENCERCORE_API SOutlinerView
 	: public STreeView<TWeakViewModelPtr<IOutlinerExtension>>
 {
 public:
@@ -139,6 +125,8 @@ public:
 	/** Scroll this tree view by the specified number of slate units */
 	void ScrollByDelta(float DeltaInSlateUnits);
 
+	bool IsColumnVisible(const FName& InName) const;
+
 protected:
 
 	/** Set the item's expansion state, including all of its children */
@@ -216,9 +204,6 @@ public:
 	/** Access all the physical nodes currently visible on the sequencer */
 	const TArray<FCachedGeometry>& GetAllVisibleNodes() const { return PhysicalNodes; }
 
-	/** Ensure that the track area column is either show or hidden, depending on the visibility of the curve editor */
-	void UpdateTrackArea();
-
 	/** Add a SOutlinerView object that should be modified or updated when this Treeview is updated */
 	void AddPinnedTreeView(TSharedPtr<SOutlinerView> PinnedTreeView);
 
@@ -229,31 +214,35 @@ public:
 	void SetShowPinned(bool bShowPinned) { bShowPinnedNodes = bShowPinned; }
 
 	/** Updates the list of visible outliner columns and regenerates columns in the outliner view */
-	void SetOutlinerColumns(const TArray<TSharedPtr<ISequencerOutlinerColumn>>& InOutlinerColumns);
+	void SetOutlinerColumns(const TArray<TSharedPtr<IOutlinerColumn>>& InOutlinerColumns);
 
 protected:
 
 	/** Linear, sorted array of nodes that we currently have generated widgets for */
 	TArray<FCachedGeometry> PhysicalNodes;
 
-	/** Helper function for creating callbacks to generate outliner columns */
-	void AddOutlinerColumnCallback(TSharedPtr<ISequencerOutlinerColumn> InColumn, TSharedPtr<FEditorViewModel> InEditorViewModel);
-
-	/** Helper function for initializing track gutter columns and padding */
-	void UpdateTrackGutterColumns(TSharedPtr<FEditorViewModel> InEditorViewModel);
+	int32 CreateOutlinerColumnsForGroup(int32 ColumnIndex, EOutlinerColumnGroup Group);
 
 	/** Populate the map of column definitions, and add relevant columns to the header row. Must be called when outliner columns change */
 	void UpdateOutlinerColumns();
 
+	/** Insert a separator column at the specified column index, with a unique identifier */
+	void InsertSeparatorColumn(int32 InsertIndex, int32 SeparatorID);
+
 	FReply OnDragRow(const FGeometry& InGeometry, const FPointerEvent& InPointerEvent, TSharedRef<SOutlinerViewRow> InRow);
 
 	FString OnItemToString_Debug(TWeakViewModelPtr<IOutlinerExtension> InWeakModel);
 
 protected:
 
+	using FColumnGenerator = TFunction<TSharedPtr<SWidget>(const FCreateOutlinerColumnParams& Params, const TSharedRef<SOutlinerViewRow>&)>;
+
 	/** The tree view's header row (hidden) */
 	TSharedPtr<SHeaderRow> HeaderRow;
 
+	/** MetaData pertaining to each column within HeaderRow */
+	TSharedPtr<FOutlinerHeaderRowWidgetMetaData> ColumnMetaData;
+
 	/** The outliner view model */
 	TWeakPtr<FOutlinerViewModel> WeakOutliner;
 
@@ -261,7 +250,7 @@ protected:
 	TArray<TWeakViewModelPtr<IOutlinerExtension>> RootNodes;
 
 	/** Column definitions for each of the columns in the tree view */
-	TMap<FName, FOutlinerViewColumn> Columns;
+	TMap<FName, FColumnGenerator> ColumnGenerators;
 
 	TSharedPtr<FSequencerCoreSelection> Selection;
 	TUniquePtr<FSelectionEventSuppressor> DelayedEventSuppressor;
@@ -276,7 +265,7 @@ protected:
 	TWeakPtr<SOutlinerView> PrimaryTreeView;
 
 	/** Visible Outliner columns to display in the outliner view */
-	TArray<TSharedPtr<ISequencerOutlinerColumn>> OutlinerColumns;
+	TArray<TSharedPtr<IOutlinerColumn>> OutlinerColumns;
 
 	float VirtualTop;
 
@@ -296,93 +285,5 @@ protected:
 	bool bRefreshPhysicalGeometry;
 };
 
-/** Widget that represents a row in the sequencer's tree control. */
-class SEQUENCERCORE_API SOutlinerViewRow
-	: public ISequencerTreeViewRow
-{
-public:
-	DECLARE_DELEGATE_RetVal_ThreeParams(TSharedRef<SWidget>, FOnGenerateWidgetForColumn, TViewModelPtr<IOutlinerExtension>, const FName&, const TSharedRef<SOutlinerViewRow>&);
-	DECLARE_DELEGATE_RetVal_ThreeParams(FReply, FDetectDrag, const FGeometry&, const FPointerEvent&, TSharedRef<SOutlinerViewRow>);
-
-	SLATE_BEGIN_ARGS(SOutlinerViewRow){}
-
-		/** Delegate to invoke to create a new column for this row */
-		SLATE_EVENT(FOnGenerateWidgetForColumn, OnGenerateWidgetForColumn)
-
-		/** Detect a drag on this tree row */
-		SLATE_EVENT(FDetectDrag, OnDetectDrag)
-
-	SLATE_END_ARGS()
-
-	/** Construct function for this widget */
-	void Construct(const FArguments& InArgs, const TSharedRef<STableViewBase>& InOwnerTableView, TWeakViewModelPtr<IOutlinerExtension> InDataModel);
-
-	/** Destroy this widget */
-	~SOutlinerViewRow();
-
-	/** Get the model to which this row relates */
-	TViewModelPtr<IOutlinerExtension> GetDataModel() const;
-
-	/**
-	 * Gets the track lane we relate to
-	 * @param bOnlyOwnTrackLane  Whether to return nullptr if our referenced track lane wasn't created
-	 *							 by our own outliner item view-model 
-	 * @return The track lane for this row
-	 */
-	TSharedPtr<STrackLane> GetTrackLane(bool bOnlyOwnTrackLane = false) const;
-
-	/**
-	 * Adds a reference to track lane, either because:
-	 *  - It is a parent row's track lane that we want to be kept alive when the parent row disappears out of view, or
-	 *  - It is a track lane that was created by our outliner view-model
-	 */
-	void SetTrackLane(const TSharedPtr<STrackLane>& InTrackLane);
-
-	/** Whether the underlying data model is selectable */
-	bool IsSelectable() const;
-
-	virtual const FSlateBrush* GetBorder() const override;
-
-	virtual void Tick(const FGeometry& AllottedGeometry, const double InCurrentTime, const float InDeltaTime) override;
-	
-	virtual int32 OnPaintDropIndicator(EItemDropZone InItemDropZone, const FPaintArgs& Args, const FGeometry& AllottedGeometry, const FSlateRect& MyCullingRect, FSlateWindowElementList& OutDrawElements, int32 LayerId, const FWidgetStyle& InWidgetStyle, bool bParentEnabled) const override;
-
-	/** Overridden from SMultiColumnTableRow.  Generates a widget for this column of the tree row. */
-	virtual TSharedRef<SWidget> GenerateWidgetForColumn(const FName& ColumnId) override;
-
-	/** Called whenever a drag is detected by the tree view. */
-	FReply OnDragDetected(const FGeometry& InGeometry, const FPointerEvent& InPointerEvent) override;
-
-	/** Called to determine whether a current drag operation is valid for this row. */
-	TOptional<EItemDropZone> OnCanAcceptDrop( const FDragDropEvent& DragDropEvent, EItemDropZone ItemDropZone, TWeakViewModelPtr<IOutlinerExtension> InDataModel);
-
-	/** Called to complete a drag and drop onto this drop. */
-	FReply OnAcceptDrop( const FDragDropEvent& DragDropEvent, EItemDropZone ItemDropZone, TWeakViewModelPtr<IOutlinerExtension> InDataModel);
-
-protected:
-
-	FVector2D ComputeDesiredSize(float LayoutScaleMultiplier) const override;
-
-protected:
-
-	/**
-	 * Cached reference to a track lane that we relate to.
-	 * Depending on the situation, this is either the actual track lane originally created
-	 * by the our outliner model, or it's just a reference we use to keep the track lane alive
-	 * (it's a weak widget) as long as we are in view.
-	 */
-	TSharedPtr<STrackLane> TrackLane;
-
-	/** The item associated with this row of data */
-	TWeakViewModelPtr<IOutlinerExtension> WeakModel;
-
-	/** Delegate to call to create a new widget for a particular column. */
-	FOnGenerateWidgetForColumn OnGenerateWidgetForColumn;
-
-	/** Delegate to call when dragging a tree item is detected */
-	FDetectDrag OnDetectDrag;
-};
-
-} // namespace Sequencer
-} // namespace UE
+} // namespace UE::Sequencer
 
diff --git a/Engine/Source/Editor/SequencerCore/Public/MVVM/Views/SOutlinerViewRow.h b/Engine/Source/Editor/SequencerCore/Public/MVVM/Views/SOutlinerViewRow.h
new file mode 100644
--- /dev/null
+++ b/Engine/Source/Editor/SequencerCore/Public/MVVM/Views/SOutlinerViewRow.h
@@ -0,0 +1,120 @@
+// Copyright Epic Games, Inc. All Rights Reserved.
+
+#pragma once
+
+#include "CoreTypes.h"
+#include "SequencerCoreFwd.h"
+#include "MVVM/Extensions/IOutlinerExtension.h"
+
+namespace UE::Sequencer
+{
+
+class STrackLane;
+
+/**
+ * Widget that represents a row in the sequencer's tree control.
+ */
+class SEQUENCERCORE_API SOutlinerViewRow
+	: public ISequencerTreeViewRow
+{
+public:
+	DECLARE_DELEGATE_RetVal_ThreeParams(TSharedRef<SWidget>, FOnGenerateWidgetForColumn, TViewModelPtr<IOutlinerExtension>, const FName&, const TSharedRef<SOutlinerViewRow>&);
+	DECLARE_DELEGATE_RetVal_ThreeParams(FReply, FDetectDrag, const FGeometry&, const FPointerEvent&, TSharedRef<SOutlinerViewRow>);
+	DECLARE_DELEGATE_RetVal_OneParam(bool, FIsColumnVisible, const FName&);
+
+	SLATE_BEGIN_ARGS(SOutlinerViewRow){}
+
+		/** Delegate to invoke to retrieve a specific column's visibility */
+		SLATE_EVENT(FIsColumnVisible, OnGetColumnVisibility)
+
+		/** Delegate to invoke to create a new column for this row */
+		SLATE_EVENT(FOnGenerateWidgetForColumn, OnGenerateWidgetForColumn)
+
+		/** Detect a drag on this tree row */
+		SLATE_EVENT(FDetectDrag, OnDetectDrag)
+
+	SLATE_END_ARGS()
+
+	/** Construct function for this widget */
+	void Construct(const FArguments& InArgs, const TSharedRef<STableViewBase>& InOwnerTableView, TWeakViewModelPtr<IOutlinerExtension> InDataModel);
+
+	/** Destroy this widget */
+	~SOutlinerViewRow();
+
+	/** Get the model to which this row relates */
+	TViewModelPtr<IOutlinerExtension> GetDataModel() const;
+
+	/**
+	 * Gets the track lane we relate to
+	 * @param bOnlyOwnTrackLane  Whether to return nullptr if our referenced track lane wasn't created
+	 *							 by our own outliner item view-model 
+	 * @return The track lane for this row
+	 */
+	TSharedPtr<STrackLane> GetTrackLane(bool bOnlyOwnTrackLane = false) const;
+
+	/**
+	 * Adds a reference to track lane, either because:
+	 *  - It is a parent row's track lane that we want to be kept alive when the parent row disappears out of view, or
+	 *  - It is a track lane that was created by our outliner view-model
+	 */
+	void SetTrackLane(const TSharedPtr<STrackLane>& InTrackLane);
+
+	/** Whether the underlying data model is selectable */
+	bool IsSelectable() const;
+
+	virtual const FSlateBrush* GetBorder() const override;
+
+	virtual void Tick(const FGeometry& AllottedGeometry, const double InCurrentTime, const float InDeltaTime) override;
+	
+	virtual void OnArrangeChildren(const FGeometry& AllottedGeometry, FArrangedChildren& ArrangedChildren) const override;
+
+	virtual int32 OnPaint(const FPaintArgs& Args, const FGeometry& AllottedGeometry, const FSlateRect& MyCullingRect, FSlateWindowElementList& OutDrawElements, int32 LayerId, const FWidgetStyle& InWidgetStyle, bool bParentEnabled) const override;
+
+	virtual int32 OnPaintDropIndicator(EItemDropZone InItemDropZone, const FPaintArgs& Args, const FGeometry& AllottedGeometry, const FSlateRect& MyCullingRect, FSlateWindowElementList& OutDrawElements, int32 LayerId, const FWidgetStyle& InWidgetStyle, bool bParentEnabled) const override;
+
+	virtual void ConstructChildren(ETableViewMode::Type InOwnerTableMode, const TAttribute<FMargin>& InPadding, const TSharedRef<SWidget>& InContent) override;
+
+	virtual bool IsColumnVisible(const FName& InColumnName) const override;
+
+	/** Called whenever a drag is detected by the tree view. */
+	FReply OnDragDetected(const FGeometry& InGeometry, const FPointerEvent& InPointerEvent) override;
+
+	/** Called to determine whether a current drag operation is valid for this row. */
+	TOptional<EItemDropZone> OnCanAcceptDrop( const FDragDropEvent& DragDropEvent, EItemDropZone ItemDropZone, TWeakViewModelPtr<IOutlinerExtension> InDataModel);
+
+	/** Called to complete a drag and drop onto this drop. */
+	FReply OnAcceptDrop( const FDragDropEvent& DragDropEvent, EItemDropZone ItemDropZone, TWeakViewModelPtr<IOutlinerExtension> InDataModel);
+
+protected:
+
+	TSharedPtr<SWidget> GenerateWidgetForColumn(const FName& ColumnId);
+
+	FVector2D ComputeDesiredSize(float LayoutScaleMultiplier) const override;
+
+protected:
+
+	/**
+	 * Cached reference to a track lane that we relate to.
+	 * Depending on the situation, this is either the actual track lane originally created
+	 * by the our outliner model, or it's just a reference we use to keep the track lane alive
+	 * (it's a weak widget) as long as we are in view.
+	 */
+	TSharedPtr<STrackLane> TrackLane;
+
+	/** The item associated with this row of data */
+	TWeakViewModelPtr<IOutlinerExtension> WeakModel;
+
+	/** Delegate to call to create a new widget for a particular column. */
+	FOnGenerateWidgetForColumn OnGenerateWidgetForColumn;
+
+	/** Delegate to call when dragging a tree item is detected */
+	FDetectDrag OnDetectDrag;
+
+	/** Delegate to call to retrieve column visibility. */
+	FIsColumnVisible OnGetColumnVisibility;
+
+	float SeparatorHeight;
+};
+
+} // namespace UE::Sequencer
+
diff --git a/Engine/Source/Editor/SequencerCore/Public/MVVM/Views/ViewUtilities.h b/Engine/Source/Editor/SequencerCore/Public/MVVM/Views/ViewUtilities.h
--- a/Engine/Source/Editor/SequencerCore/Public/MVVM/Views/ViewUtilities.h
+++ b/Engine/Source/Editor/SequencerCore/Public/MVVM/Views/ViewUtilities.h
@@ -3,20 +3,26 @@
 #pragma once
 
 #include "CoreTypes.h"
+#include "SequencerCoreFwd.h"
 #include "Framework/SlateDelegates.h"
 #include "Templates/SharedPointer.h"
-#include "Widgets/SWidget.h"
+
 
 class FText;
 class SWidget;
 template< typename ObjectType > class TAttribute;
 
-namespace UE
-{
-namespace Sequencer
+namespace UE::Sequencer
 {
 
-SEQUENCERCORE_API TSharedRef<SWidget> MakeAddButton(FText HoverText, FOnGetContent MenuContent, const TAttribute<bool>& HoverState, const TAttribute<bool>& IsEnabled);
+class IOutlinerExtension;
+
+SEQUENCERCORE_API TSharedRef<SWidget> MakeAddButton(FText HoverText, const FOnClicked& HandleClicked, const FViewModelPtr& ViewModel);
+
+SEQUENCERCORE_API TSharedRef<SWidget> MakeAddButton(FText HoverText, const FOnGetContent& HandleGetMenuContent, const FViewModelPtr& ViewModel);
+
+SEQUENCERCORE_API TSharedRef<SWidget> MakeAddButton(FText HoverText, const FOnClicked& OnClicked, const TAttribute<bool>& HoverState, const TAttribute<bool>& IsEnabled);
+
+SEQUENCERCORE_API TSharedRef<SWidget> MakeAddButton(FText HoverText, const FOnGetContent& HandleGetMenuContent, const TAttribute<bool>& HoverState, const TAttribute<bool>& IsEnabled);
 
-} // namespace Sequencer
-} // namespace UE
\ No newline at end of file
+} // namespace UE::Sequencer
\ No newline at end of file
diff --git a/Engine/Source/Runtime/Core/Public/Misc/Attribute.h b/Engine/Source/Runtime/Core/Public/Misc/Attribute.h
--- a/Engine/Source/Runtime/Core/Public/Misc/Attribute.h
+++ b/Engine/Source/Runtime/Core/Public/Misc/Attribute.h
@@ -498,6 +498,17 @@ template<typename T, typename SourceType, typename SourceTypeOrBase, typename...
 {
 	return TAttribute<T>::Create(TAttribute<T>::FGetter::CreateSP(InObject, InMethod, MoveTemp(InputPayload)...));
 }
+template<typename T, typename SourceType, typename SourceTypeOrBase, typename... PayloadTypes>
+[[nodiscard]] FORCEINLINE TAttribute<T> MakeAttributeSP(TSharedRef<SourceType> InObject, T (SourceTypeOrBase::*InMethod)(PayloadTypes...) const, typename TDecay<PayloadTypes>::Type...  InputPayload)
+{
+	return TAttribute<T>::Create(TAttribute<T>::FGetter::CreateSP(InObject, InMethod, MoveTemp(InputPayload)...));
+}
+
+template<typename T, typename SourceType, typename SourceTypeOrBase, typename... PayloadTypes>
+[[nodiscard]] FORCEINLINE TAttribute<T> MakeAttributeUObject(const SourceType* InObject, T (SourceTypeOrBase::*InMethod)(PayloadTypes...) const, typename TDecay<PayloadTypes>::Type... InputPayload)
+{
+	return TAttribute<T>::Create(TAttribute<T>::FGetter::CreateUObject(InObject, InMethod, MoveTemp(InputPayload)...));
+}
 
 /**
  * Helper function for creating TAttributes from a lambda
diff --git a/Engine/Source/Runtime/Slate/Public/Widgets/Views/STableRow.h b/Engine/Source/Runtime/Slate/Public/Widgets/Views/STableRow.h
--- a/Engine/Source/Runtime/Slate/Public/Widgets/Views/STableRow.h
+++ b/Engine/Source/Runtime/Slate/Public/Widgets/Views/STableRow.h
@@ -321,7 +321,7 @@ public:
 		};
 	}
 
-	virtual int32 OnPaint( const FPaintArgs& Args, const FGeometry& AllottedGeometry, const FSlateRect& MyCullingRect, FSlateWindowElementList& OutDrawElements, int32 LayerId, const FWidgetStyle& InWidgetStyle, bool bParentEnabled ) const override
+	int32 PaintSelection( const FPaintArgs& Args, const FGeometry& AllottedGeometry, const FSlateRect& MyCullingRect, FSlateWindowElementList& OutDrawElements, int32 LayerId, const FWidgetStyle& InWidgetStyle, bool bParentEnabled ) const
 	{
 		TSharedRef< ITypedTableView<ItemType> > OwnerTable = OwnerTablePtr.Pin().ToSharedRef();
 		const bool bIsActive = OwnerTable->AsWidget()->HasKeyboardFocus();
@@ -332,7 +332,7 @@ public:
 			{
 				FSlateDrawElement::MakeBox(
 					OutDrawElements,
-					LayerId,
+					LayerId++,
 					AllottedGeometry.ToPaintGeometry(),
 					&Style->SelectorFocusedBrush,
 					ESlateDrawEffect::None,
@@ -340,24 +340,39 @@ public:
 				);
 			}
 		}
-
-		LayerId = SBorder::OnPaint(Args, AllottedGeometry, MyCullingRect, OutDrawElements, LayerId, InWidgetStyle, bParentEnabled);
-
+		return LayerId;
+	}
+	int32 PaintBorder( const FPaintArgs& Args, const FGeometry& AllottedGeometry, const FSlateRect& MyCullingRect, FSlateWindowElementList& OutDrawElements, int32 LayerId, const FWidgetStyle& InWidgetStyle, bool bParentEnabled ) const
+	{
+		return SBorder::OnPaint(Args, AllottedGeometry, MyCullingRect, OutDrawElements, LayerId, InWidgetStyle, bParentEnabled);
+	}
+	int32 PaintDropIndicator( const FPaintArgs& Args, const FGeometry& AllottedGeometry, const FSlateRect& MyCullingRect, FSlateWindowElementList& OutDrawElements, int32 LayerId, const FWidgetStyle& InWidgetStyle, bool bParentEnabled ) const
+	{
 		if (ItemDropZone.IsSet())
 		{
 			if (PaintDropIndicatorEvent.IsBound())
 			{
-				LayerId = PaintDropIndicatorEvent.Execute(ItemDropZone.GetValue(), Args, AllottedGeometry, MyCullingRect, OutDrawElements, LayerId, InWidgetStyle, bParentEnabled);
+				return PaintDropIndicatorEvent.Execute(ItemDropZone.GetValue(), Args, AllottedGeometry, MyCullingRect, OutDrawElements, LayerId, InWidgetStyle, bParentEnabled);
 			}
 			else
 			{
-				OnPaintDropIndicator(ItemDropZone.GetValue(), Args, AllottedGeometry, MyCullingRect, OutDrawElements, LayerId, InWidgetStyle, bParentEnabled);
+				return OnPaintDropIndicator(ItemDropZone.GetValue(), Args, AllottedGeometry, MyCullingRect, OutDrawElements, LayerId, InWidgetStyle, bParentEnabled);
 			}
 		}
 
 		return LayerId;
 	}
 
+
+	virtual int32 OnPaint( const FPaintArgs& Args, const FGeometry& AllottedGeometry, const FSlateRect& MyCullingRect, FSlateWindowElementList& OutDrawElements, int32 LayerId, const FWidgetStyle& InWidgetStyle, bool bParentEnabled ) const override
+	{
+		LayerId = PaintSelection(Args, AllottedGeometry, MyCullingRect, OutDrawElements, LayerId, InWidgetStyle, bParentEnabled);
+		LayerId = PaintBorder(Args, AllottedGeometry, MyCullingRect, OutDrawElements, LayerId, InWidgetStyle, bParentEnabled);
+		LayerId = PaintDropIndicator(Args, AllottedGeometry, MyCullingRect, OutDrawElements, LayerId, InWidgetStyle, bParentEnabled);
+
+		return LayerId;
+	}
+
 	virtual int32 OnPaintDropIndicator( EItemDropZone InItemDropZone, const FPaintArgs& Args, const FGeometry& AllottedGeometry, const FSlateRect& MyCullingRect, FSlateWindowElementList& OutDrawElements, int32 LayerId, const FWidgetStyle& InWidgetStyle, bool bParentEnabled ) const
 	{
 		TSharedRef< ITypedTableView<ItemType> > OwnerTable = OwnerTablePtr.Pin().ToSharedRef();
