Compile fix

diff --git a/Engine/Source/Editor/SequencerCore/Private/MVVM/Views/SOutlinerView.cpp b/Engine/Source/Editor/SequencerCore/Private/MVVM/Views/SOutlinerView.cpp
--- a/Engine/Source/Editor/SequencerCore/Private/MVVM/Views/SOutlinerView.cpp
+++ b/Engine/Source/Editor/SequencerCore/Private/MVVM/Views/SOutlinerView.cpp
@@ -335,14 +335,14 @@ void SOutlinerView::InsertSeparatorColumn(int32 InsertIndex, int32 SeparatorID)
 {
 	static const FName NAME_Separator("Separator");
 
-	FOutlinerColumnLayout SeparatorLayout(
+	FOutlinerColumnLayout SeparatorLayout{
 		1.f, /* Width */
 		FMargin(0.f),
 		HAlign_Fill,
 		VAlign_Fill,
 		EOutlinerColumnSizeMode::Fixed,
 		EOutlinerColumnFlags::Hidden
-	);
+	};
 	FName SeparatorName(NAME_Separator, SeparatorID);
 
 	// Add a 1px separator column
