Sequencer: Corrected hover state handling for outliner rows

diff --git a/Engine/Source/Editor/SequencerCore/Private/MVVM/Views/SOutlinerItemViewBase.cpp b/Engine/Source/Editor/SequencerCore/Private/MVVM/Views/SOutlinerItemViewBase.cpp
--- a/Engine/Source/Editor/SequencerCore/Private/MVVM/Views/SOutlinerItemViewBase.cpp
+++ b/Engine/Source/Editor/SequencerCore/Private/MVVM/Views/SOutlinerItemViewBase.cpp
@@ -61,7 +61,6 @@ void SOutlinerItemViewBase::Construct(
 
 	ItemStyle              = InArgs._ItemStyle;
 	IsReadOnlyAttribute    = InArgs._IsReadOnly;
-	IsRowHoveredAttribute  = MakeAttributeSP(&InTableRow.Get(), &ISequencerTreeViewRow::IsHovered);
 	IsRowSelectedAttribute = MakeAttributeSP(&InTableRow.Get(), &ISequencerTreeViewRow::IsItemSelected);
 
 	if (!IsReadOnlyAttribute.IsSet())
@@ -363,28 +362,6 @@ FOptionalSize SOutlinerItemViewBase::GetHeight() const
 		: 10.f;
 }
 
-void SOutlinerItemViewBase::OnMouseEnter(const FGeometry& MyGeometry, const FPointerEvent& MouseEvent)
-{
-	TViewModelPtr<IOutlinerExtension> DataModel = WeakOutlinerExtension.Pin();
-	TSharedPtr<FEditorViewModel> Editor = WeakEditor.Pin();
-	if (DataModel && Editor)
-	{
-		Editor->GetOutliner()->SetHoveredItem(DataModel);
-	}
-	SWidget::OnMouseEnter(MyGeometry, MouseEvent);
-}
-
-void SOutlinerItemViewBase::OnMouseLeave(const FPointerEvent& MouseEvent)
-{
-	TSharedPtr<FEditorViewModel> Editor = WeakEditor.Pin();
-	if (Editor)
-	{
-		Editor->GetOutliner()->SetHoveredItem(nullptr);
-	}
-
-	SWidget::OnMouseLeave(MouseEvent);
-}
-
 const FSlateBrush* SOutlinerItemViewBase::GetNodeBorderImage() const
 {
 	TSharedPtr<FViewModel> DataModel = WeakOutlinerExtension.Pin().AsModel();
diff --git a/Engine/Source/Editor/SequencerCore/Private/MVVM/Views/SOutlinerViewRow.cpp b/Engine/Source/Editor/SequencerCore/Private/MVVM/Views/SOutlinerViewRow.cpp
--- a/Engine/Source/Editor/SequencerCore/Private/MVVM/Views/SOutlinerViewRow.cpp
+++ b/Engine/Source/Editor/SequencerCore/Private/MVVM/Views/SOutlinerViewRow.cpp
@@ -5,6 +5,9 @@
 #include "MVVM/Views/SOutlinerView.h"
 #include "MVVM/Views/STrackLane.h"
 #include "MVVM/Extensions/ISelectableExtension.h"
+#include "MVVM/ViewModels/EditorViewModel.h"
+#include "MVVM/ViewModels/OutlinerViewModel.h"
+#include "MVVM/ViewModels/EditorSharedViewModelData.h"
 
 
 namespace UE::Sequencer
@@ -237,6 +240,36 @@ TViewModelPtr<IOutlinerExtension> SOutlinerViewRow::GetDataModel() const
 	return WeakModel.Pin();
 }
 
+void SOutlinerViewRow::OnMouseEnter(const FGeometry& MyGeometry, const FPointerEvent& MouseEvent)
+{
+	TViewModelPtr<IOutlinerExtension>       DataModel        = WeakModel.Pin();
+	TSharedPtr<FSharedViewModelData>        SharedData       = DataModel        ? DataModel.AsModel()->GetSharedData() : nullptr;
+	TSharedPtr<FEditorSharedViewModelData>  SharedEditorData = SharedData       ? SharedData->CastThisShared<FEditorSharedViewModelData>() : nullptr;
+	TSharedPtr<FEditorViewModel>            Editor           = SharedEditorData ? SharedEditorData->GetEditor() : nullptr;
+
+	if (DataModel && Editor)
+	{
+		Editor->GetOutliner()->SetHoveredItem(DataModel);
+	}
+	SWidget::OnMouseEnter(MyGeometry, MouseEvent);
+}
+
+void SOutlinerViewRow::OnMouseLeave(const FPointerEvent& MouseEvent)
+{
+	TViewModelPtr<IOutlinerExtension>       DataModel        = WeakModel.Pin();
+	TSharedPtr<FSharedViewModelData>        SharedData       = DataModel ? DataModel.AsModel()->GetSharedData() : nullptr;
+	TSharedPtr<FEditorSharedViewModelData>  SharedEditorData = SharedData       ? SharedData->CastThisShared<FEditorSharedViewModelData>() : nullptr;
+	TSharedPtr<FEditorViewModel>            Editor           = SharedEditorData ? SharedEditorData->GetEditor() : nullptr;
+
+	if (Editor)
+	{
+		Editor->GetOutliner()->SetHoveredItem(nullptr);
+	}
+
+	SWidget::OnMouseLeave(MouseEvent);
+}
+
+
 TSharedPtr<STrackLane> SOutlinerViewRow::GetTrackLane(bool bOnlyOwnTrackLane) const
 {
 	if (!bOnlyOwnTrackLane)
diff --git a/Engine/Source/Editor/SequencerCore/Public/MVVM/Views/SOutlinerItemViewBase.h b/Engine/Source/Editor/SequencerCore/Public/MVVM/Views/SOutlinerItemViewBase.h
--- a/Engine/Source/Editor/SequencerCore/Public/MVVM/Views/SOutlinerItemViewBase.h
+++ b/Engine/Source/Editor/SequencerCore/Public/MVVM/Views/SOutlinerItemViewBase.h
@@ -129,9 +129,6 @@ private:
 
 	// SWidget interface
 
-	void OnMouseEnter(const FGeometry& MyGeometry, const FPointerEvent& MouseEvent) override;
-	void OnMouseLeave(const FPointerEvent& MouseEvent) override;
-
 	FSlateColor GetForegroundBasedOnSelection() const;
 
 	/**
@@ -160,7 +157,6 @@ protected:
 	TWeakPtr<FEditorViewModel> WeakEditor;
 
 	TAttribute<bool> IsReadOnlyAttribute;
-	TAttribute<bool> IsRowHoveredAttribute;
 	TAttribute<bool> IsRowSelectedAttribute;
 
 	/** Default background brush for this node when expanded */
diff --git a/Engine/Source/Editor/SequencerCore/Public/MVVM/Views/SOutlinerViewRow.h b/Engine/Source/Editor/SequencerCore/Public/MVVM/Views/SOutlinerViewRow.h
--- a/Engine/Source/Editor/SequencerCore/Public/MVVM/Views/SOutlinerViewRow.h
+++ b/Engine/Source/Editor/SequencerCore/Public/MVVM/Views/SOutlinerViewRow.h
@@ -74,6 +74,10 @@ public:
 
 	virtual void ConstructChildren(ETableViewMode::Type InOwnerTableMode, const TAttribute<FMargin>& InPadding, const TSharedRef<SWidget>& InContent) override;
 
+	virtual void OnMouseEnter(const FGeometry& MyGeometry, const FPointerEvent& MouseEvent) override;
+
+	virtual void OnMouseLeave(const FPointerEvent& MouseEvent) override;
+
 	virtual bool IsColumnVisible(const FName& InColumnName) const override;
 
 	/** Called whenever a drag is detected by the tree view. */
