Sequencer: Change outliner column layout algorithm to encroach cell padding without hiding the subsequent cell if possible

diff --git a/Engine/Source/Editor/SequencerCore/Private/MVVM/Views/SOutlinerViewRowPanel.cpp b/Engine/Source/Editor/SequencerCore/Private/MVVM/Views/SOutlinerViewRowPanel.cpp
--- a/Engine/Source/Editor/SequencerCore/Private/MVVM/Views/SOutlinerViewRowPanel.cpp
+++ b/Engine/Source/Editor/SequencerCore/Private/MVVM/Views/SOutlinerViewRowPanel.cpp
@@ -104,6 +104,7 @@ void SOutlinerViewRowPanel::OnArrangeChildren(const FGeometry& AllottedGeometry,
 		}
 
 		float CellWidth = ComputeCellWidth(Column, Layout);
+		float PaddingEncroachment = 0.f;
 
 		// If this column is for the next slot, handle any overflowing required and arrange the widget
 		if (SlotIndex < NumSlots && Children[SlotIndex].ColumnId == Column.ColumnId)
@@ -128,6 +129,14 @@ void SOutlinerViewRowPanel::OnArrangeChildren(const FGeometry& AllottedGeometry,
 					const SHeaderRow::FColumn&   NextColumn = Columns[NextColumnIndex];
 					const FOutlinerColumnLayout& NextLayout = MetaDataEntry->Columns[NextColumnIndex];
 
+					// Encroach the next cell's padding if possible without hiding the slot
+					if (CellWidth + NextLayout.CellPadding.Left >= RequiredWidth)
+					{
+						CellWidth += NextLayout.CellPadding.Left;
+						PaddingEncroachment = NextLayout.CellPadding.Left;
+						break;
+					}
+
 					CellWidth += ComputeCellWidth(NextColumn, NextLayout);
 
 					// If this belonged to the next slot, move past that as well
@@ -198,7 +207,7 @@ void SOutlinerViewRowPanel::OnArrangeChildren(const FGeometry& AllottedGeometry,
 			++SlotIndex;
 		}
 
-		HorizontalOffset += CellWidth;
+		HorizontalOffset += CellWidth - PaddingEncroachment;
 	}
 }
 
