Sequencer: Nav and keyframe columns no longer overflow the color picker

This ensures that widgets placed in those columns are always aligned with each other

diff --git a/Engine/Source/Editor/SequencerCore/Private/MVVM/ViewModels/OutlinerColumns/KeyFrameOutlinerColumn.cpp b/Engine/Source/Editor/SequencerCore/Private/MVVM/ViewModels/OutlinerColumns/KeyFrameOutlinerColumn.cpp
--- a/Engine/Source/Editor/SequencerCore/Private/MVVM/ViewModels/OutlinerColumns/KeyFrameOutlinerColumn.cpp
+++ b/Engine/Source/Editor/SequencerCore/Private/MVVM/ViewModels/OutlinerColumns/KeyFrameOutlinerColumn.cpp
@@ -10,7 +10,7 @@ FKeyFrameOutlinerColumn::FKeyFrameOutlinerColumn()
 	Name     = FCommonOutlinerNames::KeyFrame;
 	Label    = NSLOCTEXT("FKeyFrameOutlinerColumn", "ColumnName", "Key Frame");
 	Position = FOutlinerColumnPosition{ 20, EOutlinerColumnGroup::RightGutter };
-	Layout   = FOutlinerColumnLayout{ 34.f, FMargin(0.f), HAlign_Center, VAlign_Center, EOutlinerColumnSizeMode::Fixed, EOutlinerColumnFlags::OverflowSubsequentEmptyCells };
+	Layout   = FOutlinerColumnLayout{ 34.f, FMargin(0.f), HAlign_Center, VAlign_Center, EOutlinerColumnSizeMode::Fixed };
 }
 
 bool FKeyFrameOutlinerColumn::IsColumnVisibleByDefault() const
diff --git a/Engine/Source/Editor/SequencerCore/Private/MVVM/ViewModels/OutlinerColumns/NavOutlinerColumn.cpp b/Engine/Source/Editor/SequencerCore/Private/MVVM/ViewModels/OutlinerColumns/NavOutlinerColumn.cpp
--- a/Engine/Source/Editor/SequencerCore/Private/MVVM/ViewModels/OutlinerColumns/NavOutlinerColumn.cpp
+++ b/Engine/Source/Editor/SequencerCore/Private/MVVM/ViewModels/OutlinerColumns/NavOutlinerColumn.cpp
@@ -10,7 +10,7 @@ FNavOutlinerColumn::FNavOutlinerColumn()
 	Name     = FCommonOutlinerNames::Nav;
 	Label    = NSLOCTEXT("FNavOutlinerColumn", "ColumnName", "Key Navigation");
 	Position = FOutlinerColumnPosition{ 10, EOutlinerColumnGroup::RightGutter };
-	Layout   = FOutlinerColumnLayout{ 74.f, FMargin(0.f), HAlign_Center, VAlign_Center, EOutlinerColumnSizeMode::Fixed, EOutlinerColumnFlags::OverflowSubsequentEmptyCells };
+	Layout   = FOutlinerColumnLayout{ 74.f, FMargin(0.f), HAlign_Center, VAlign_Center, EOutlinerColumnSizeMode::Fixed };
 }
 
 } // namespace UE::Sequencer
