[SequencerCore] Added full namespace qualifiers to some macros to allow use outside of UE::Sequencer

diff --git a/Engine/Source/Editor/SequencerCore/Public/MVVM/ICastable.h b/Engine/Source/Editor/SequencerCore/Public/MVVM/ICastable.h
--- a/Engine/Source/Editor/SequencerCore/Public/MVVM/ICastable.h
+++ b/Engine/Source/Editor/SequencerCore/Public/MVVM/ICastable.h
@@ -130,12 +130,12 @@ struct TImplements
 {
 };
 
-#define UE_SEQUENCER_DECLARE_CASTABLE(ThisType, ...)                                             \
-	UE_SEQUENCER_DECLARE_VIEW_MODEL_TYPE_ID(ThisType);                                           \
-	using Implements = TImplements<__VA_ARGS__>;                                                 \
-	UE_NO_UNIQUE_ADDRESS FAutoRegisterTypeTable AutoRegisterTypeTable = FAutoRegisterTypeTable(static_cast<ThisType*>(this)); \
+#define UE_SEQUENCER_DECLARE_CASTABLE(ThisType, ...)																																\
+	UE_SEQUENCER_DECLARE_VIEW_MODEL_TYPE_ID(ThisType);																																\
+	using Implements = ::UE::Sequencer::TImplements<__VA_ARGS__>;																														\
+	UE_NO_UNIQUE_ADDRESS ::UE::Sequencer::ICastable::FAutoRegisterTypeTable AutoRegisterTypeTable = ::UE::Sequencer::ICastable::FAutoRegisterTypeTable(static_cast<ThisType*>(this));	\
 
-#define UE_SEQUENCER_DEFINE_CASTABLE(ThisType)                                                   \
+#define UE_SEQUENCER_DEFINE_CASTABLE(ThisType)																																		\
 	UE_SEQUENCER_DEFINE_VIEW_MODEL_TYPE_ID(ThisType);
 
 
diff --git a/Engine/Source/Editor/SequencerCore/Public/MVVM/ViewModelTypeID.h b/Engine/Source/Editor/SequencerCore/Public/MVVM/ViewModelTypeID.h
--- a/Engine/Source/Editor/SequencerCore/Public/MVVM/ViewModelTypeID.h
+++ b/Engine/Source/Editor/SequencerCore/Public/MVVM/ViewModelTypeID.h
@@ -5,20 +5,20 @@
 #include "CoreTypes.h"
 #include "MVVM/CastableTypeTable.h"
 
-#define UE_SEQUENCER_DECLARE_VIEW_MODEL_TYPE_ID(Type)                                               \
-	static ::UE::Sequencer::TAutoRegisterViewModelTypeID<Type> ID;                                  \
+#define UE_SEQUENCER_DECLARE_VIEW_MODEL_TYPE_ID(Type)																	\
+	static ::UE::Sequencer::TAutoRegisterViewModelTypeID<Type> ID;														\
 	static void RegisterTypeID();
 
-#define UE_SEQUENCER_DECLARE_VIEW_MODEL_TYPE_ID_API(MODULE_API, Type)                               \
-	MODULE_API static ::UE::Sequencer::TAutoRegisterViewModelTypeID<Type> ID;                       \
+#define UE_SEQUENCER_DECLARE_VIEW_MODEL_TYPE_ID_API(MODULE_API, Type)													\
+	MODULE_API static ::UE::Sequencer::TAutoRegisterViewModelTypeID<Type> ID;											\
 	MODULE_API static void RegisterTypeID();
 
-#define UE_SEQUENCER_DEFINE_VIEW_MODEL_TYPE_ID(Type)                                                \
-	::UE::Sequencer::TAutoRegisterViewModelTypeID<Type> Type::ID;                                   \
-	void Type::RegisterTypeID()                                                                     \
-	{                                                                                               \
-		Type::ID.ID        = FViewModelTypeID::RegisterNewID();                                     \
-		Type::ID.TypeTable = FCastableTypeTable::MakeTypeTable<Type>((Type*)0, Type::ID.ID, #Type); \
+#define UE_SEQUENCER_DEFINE_VIEW_MODEL_TYPE_ID(Type)																	\
+	::UE::Sequencer::TAutoRegisterViewModelTypeID<Type> Type::ID;														\
+	void Type::RegisterTypeID()																							\
+	{																													\
+		Type::ID.ID        = ::UE::Sequencer::FViewModelTypeID::RegisterNewID();										\
+		Type::ID.TypeTable = ::UE::Sequencer::FCastableTypeTable::MakeTypeTable<Type>((Type*)0, Type::ID.ID, #Type);	\
 	}
 
 namespace UE
