Sequencer: Fix regression and crash with add button not being disabled when the sequencer is read only

#rnx

diff --git a/Engine/Source/Editor/Sequencer/Private/Sequencer.cpp b/Engine/Source/Editor/Sequencer/Private/Sequencer.cpp
--- a/Engine/Source/Editor/Sequencer/Private/Sequencer.cpp
+++ b/Engine/Source/Editor/Sequencer/Private/Sequencer.cpp
@@ -1718,6 +1718,12 @@ void FSequencer::GetConstructionScriptActors(UMovieScene* MovieScene, FMovieScen
 
 void FSequencer::DeleteSections(const TSet<UMovieSceneSection*>& Sections)
 {
+	if (IsReadOnly())
+	{
+		FSequencerUtilities::ShowReadOnlyError();
+		return;
+	}
+
 	UMovieScene* MovieScene = GetFocusedMovieSceneSequence()->GetMovieScene();
 	bool bAnythingRemoved = false;
 
@@ -2235,6 +2241,12 @@ void FSequencer::TransformSelectedKeysAndSections(FFrameTime InDeltaTime, float
 
 void FSequencer::TranslateSelectedKeysAndSections(bool bTranslateLeft)
 {
+	if (IsReadOnly())
+	{
+		FSequencerUtilities::ShowReadOnlyError();
+		return;
+	}
+
 	int32 Shift = bTranslateLeft ? -1 : 1;
 	FFrameTime Delta = FQualifiedFrameTime(Shift, GetFocusedDisplayRate()).ConvertTo(GetFocusedTickResolution());
 	TransformSelectedKeysAndSections(Delta, 1.f);
diff --git a/Engine/Source/Editor/SequencerCore/Private/MVVM/Views/ViewUtilities.cpp b/Engine/Source/Editor/SequencerCore/Private/MVVM/Views/ViewUtilities.cpp
--- a/Engine/Source/Editor/SequencerCore/Private/MVVM/Views/ViewUtilities.cpp
+++ b/Engine/Source/Editor/SequencerCore/Private/MVVM/Views/ViewUtilities.cpp
@@ -6,6 +6,7 @@
 #include "MVVM/ViewModelPtr.h"
 #include "MVVM/ViewModels/ViewModelIterators.h"
 #include "MVVM/ViewModels/EditorViewModel.h"
+#include "MVVM/ViewModels/EditorSharedViewModelData.h"
 #include "MVVM/Extensions/IHoveredExtension.h"
 #include "MVVM/Views/OutlinerColumns/SOutlinerColumnButton.h"
 
@@ -21,9 +22,12 @@ TSharedRef<SWidget> MakeAddButton(FText HoverText, const FOnClicked& HandleClick
 		IsHovered = MakeAttributeSP(Hoverable.ToSharedRef(), &IHoveredExtension::IsHovered);
 	}
 
+	TSharedPtr<FSharedViewModelData>        SharedData       = ViewModel->GetSharedData();
+	TSharedPtr<FEditorSharedViewModelData>  SharedEditorData = SharedData       ? SharedData->CastThisShared<FEditorSharedViewModelData>() : nullptr;
+	TSharedPtr<FEditorViewModel>            Editor           = SharedEditorData ? SharedEditorData->GetEditor() : nullptr;
 
 	TAttribute<bool> IsEnabled;
-	if (TSharedPtr<FEditorViewModel> Editor = ViewModel->FindAncestorOfType<FEditorViewModel>())
+	if (Editor)
 	{
 		IsEnabled = MakeAttributeSP(Editor.ToSharedRef(), &FEditorViewModel::IsEditable);
 	}
@@ -44,8 +48,12 @@ TSharedRef<SWidget> MakeAddButton(FText HoverText, const FOnGetContent& HandleGe
 		IsHovered = MakeAttributeSP(Hoverable.ToSharedRef(), &IHoveredExtension::IsHovered);
 	}
 
+	TSharedPtr<FSharedViewModelData>        SharedData       = ViewModel->GetSharedData();
+	TSharedPtr<FEditorSharedViewModelData>  SharedEditorData = SharedData       ? SharedData->CastThisShared<FEditorSharedViewModelData>() : nullptr;
+	TSharedPtr<FEditorViewModel>            Editor           = SharedEditorData ? SharedEditorData->GetEditor() : nullptr;
+
 	TAttribute<bool> IsEnabled;
-	if (TSharedPtr<FEditorViewModel> Editor = ViewModel->FindAncestorOfType<FEditorViewModel>())
+	if (Editor)
 	{
 		IsEnabled = MakeAttributeSP(Editor.ToSharedRef(), &FEditorViewModel::IsEditable);
 	}
