Adding GetPath to FSequencerViewModelScriptingStruct for selection use in python

[RN] Adds support for retrieving the path for a selection from sequencer's outliner. This API is also exposed to blueprint/python.

diff --git a/Engine/Source/Editor/SequencerCore/Private/Scripting/ViewModelScriptingStruct.cpp b/Engine/Source/Editor/SequencerCore/Private/Scripting/ViewModelScriptingStruct.cpp
--- a/Engine/Source/Editor/SequencerCore/Private/Scripting/ViewModelScriptingStruct.cpp
+++ b/Engine/Source/Editor/SequencerCore/Private/Scripting/ViewModelScriptingStruct.cpp
@@ -27,4 +27,12 @@ FString USequencerViewModelStructExtensions::GetLabel(const FSequencerViewModelS
 
 	TViewModelPtr<IOutlinerExtension> Outliner = ViewModel.WeakViewModel.ImplicitPin();
 	return Outliner ? Outliner->GetLabel().ToString() : FString();
-}
\ No newline at end of file
+}
+
+FString USequencerViewModelStructExtensions::GetPath(const FSequencerViewModelScriptingStruct& ViewModel)
+{
+	using namespace UE::Sequencer;
+
+	TViewModelPtr<IOutlinerExtension> Outliner = ViewModel.WeakViewModel.ImplicitPin();
+	return Outliner ? Outliner->GetPathName(Outliner.AsModel()) : FString();
+}
diff --git a/Engine/Source/Editor/SequencerCore/Public/Scripting/ViewModelScriptingStruct.h b/Engine/Source/Editor/SequencerCore/Public/Scripting/ViewModelScriptingStruct.h
--- a/Engine/Source/Editor/SequencerCore/Public/Scripting/ViewModelScriptingStruct.h
+++ b/Engine/Source/Editor/SequencerCore/Public/Scripting/ViewModelScriptingStruct.h
@@ -22,7 +22,6 @@ struct FSequencerViewModelScriptingStruct
 	{
 		return A.WeakViewModel == B.WeakViewModel;
 	}
-
 	UPROPERTY(VisibleAnywhere, BlueprintReadOnly, Category="Sequencer Editor")
 	FName Type;
 
@@ -52,4 +51,7 @@ public:
 
 	UFUNCTION(BlueprintPure, Category = "Sequencer Editor", meta=(ScriptMethod))
 	static FString GetLabel(const FSequencerViewModelScriptingStruct& ViewModel);
+
+	UFUNCTION(BlueprintPure, Category = "Sequencer Editor", meta = (ScriptMethod))
+	static FString GetPath(const FSequencerViewModelScriptingStruct& ViewModel);
 };
\ No newline at end of file
