Sequencer: Deprecate FViewModelHierarchyOperation constructor that accepts view models directly

This constructor is now actually ambiguous with the shared data constructor ever since shared data became a view model itself. Since it is used infrequently, I am deprecating it in favor of the shared data constructor.

diff --git a/Engine/Source/Editor/Sequencer/Private/MVVM/ViewModels/SequenceModel.cpp b/Engine/Source/Editor/Sequencer/Private/MVVM/ViewModels/SequenceModel.cpp
--- a/Engine/Source/Editor/Sequencer/Private/MVVM/ViewModels/SequenceModel.cpp
+++ b/Engine/Source/Editor/Sequencer/Private/MVVM/ViewModels/SequenceModel.cpp
@@ -121,7 +121,7 @@ void FSequenceModel::SetSequence(UMovieSceneSequence* InSequence, FMovieSceneSeq
 void FSequenceModel::OnPostUndo()
 {
 	{
-		FViewModelHierarchyOperation HierarchyOp(AsShared());
+		FViewModelHierarchyOperation HierarchyOp(GetSharedData());
 
 		for (TSharedPtr<FViewModel> Child : GetChildren().ToArray())
 		{
diff --git a/Engine/Source/Editor/Sequencer/Private/SequencerNodeTree.cpp b/Engine/Source/Editor/Sequencer/Private/SequencerNodeTree.cpp
--- a/Engine/Source/Editor/Sequencer/Private/SequencerNodeTree.cpp
+++ b/Engine/Source/Editor/Sequencer/Private/SequencerNodeTree.cpp
@@ -116,7 +116,7 @@ void FSequencerNodeTree::Update()
 {
 	using namespace UE::Sequencer;
 
-	FViewModelHierarchyOperation UpdateOp(RootNode);
+	FViewModelHierarchyOperation UpdateOp(RootNode->GetSharedData());
 
 	FilteredNodes.Empty();
 
diff --git a/Engine/Source/Editor/SequencerCore/Private/MVVM/Extensions/ISortableExtension.cpp b/Engine/Source/Editor/SequencerCore/Private/MVVM/Extensions/ISortableExtension.cpp
--- a/Engine/Source/Editor/SequencerCore/Private/MVVM/Extensions/ISortableExtension.cpp
+++ b/Engine/Source/Editor/SequencerCore/Private/MVVM/Extensions/ISortableExtension.cpp
@@ -58,7 +58,7 @@ void ISortableExtension::SortChildren(TSharedPtr<FViewModel> ParentModel, ESorti
 
 void ISortableExtension::SortChildren(FViewModelChildren& Children, ESortingMode SortingMode)
 {
-	FViewModelHierarchyOperation Operation(Children.GetParent());
+	FViewModelHierarchyOperation Operation(Children.GetParent()->GetSharedData());
 
 	struct FModelAndSortingKey
 	{
diff --git a/Engine/Source/Editor/SequencerCore/Private/MVVM/ViewModels/ViewModelHierarchy.cpp b/Engine/Source/Editor/SequencerCore/Private/MVVM/ViewModels/ViewModelHierarchy.cpp
--- a/Engine/Source/Editor/SequencerCore/Private/MVVM/ViewModels/ViewModelHierarchy.cpp
+++ b/Engine/Source/Editor/SequencerCore/Private/MVVM/ViewModels/ViewModelHierarchy.cpp
@@ -354,7 +354,7 @@ void FViewModelChildren::Empty()
 {
 	if (ListHead->HeadLink.Next)
 	{
-		FViewModelHierarchyOperation Operation(Owner);
+		FViewModelHierarchyOperation Operation(Owner->GetSharedData());
 
 		for (const FViewModelPtr& Child : IterateSubList().ToArray())
 		{
@@ -365,13 +365,7 @@ void FViewModelChildren::Empty()
 	}
 }
 
-FViewModelHierarchyOperation::FViewModelHierarchyOperation(const TSharedPtr<FViewModel>& InAnyModel)
-	: SharedData(InAnyModel->SharedData)
-{
-	Construct();
-}
-
-FViewModelHierarchyOperation::FViewModelHierarchyOperation(const TSharedRef<FSharedViewModelData>& InSharedData)
+FViewModelHierarchyOperation::FViewModelHierarchyOperation(const TSharedPtr<FSharedViewModelData>& InSharedData)
 	: SharedData(InSharedData)
 {
 	Construct();
diff --git a/Engine/Source/Editor/SequencerCore/Public/MVVM/ViewModels/ViewModelHierarchy.h b/Engine/Source/Editor/SequencerCore/Public/MVVM/ViewModels/ViewModelHierarchy.h
--- a/Engine/Source/Editor/SequencerCore/Public/MVVM/ViewModels/ViewModelHierarchy.h
+++ b/Engine/Source/Editor/SequencerCore/Public/MVVM/ViewModels/ViewModelHierarchy.h
@@ -290,8 +290,16 @@ private:
 class SEQUENCERCORE_API FViewModelHierarchyOperation
 {
 public:
-	explicit FViewModelHierarchyOperation(const TSharedPtr<FViewModel>& InAnyModel);
-	explicit FViewModelHierarchyOperation(const TSharedRef<FSharedViewModelData>& InSharedData);
+
+	template<typename ViewModelType>
+	UE_DEPRECATED(5.5, "Please use the FSharedViewModelData overload by passing ->GetSharedData()")
+	explicit FViewModelHierarchyOperation(const TSharedPtr<ViewModelType>& InAnyModel)
+		: SharedData(InAnyModel->SharedData)
+	{
+		Construct();
+	}
+
+	explicit FViewModelHierarchyOperation(const TSharedPtr<FSharedViewModelData>& InSharedData);
 	~FViewModelHierarchyOperation();
 
 	/**
