SequencerCore: Added AsWeak methods to TViewModelPtr

diff --git a/Engine/Source/Editor/SequencerCore/Private/MVVM/ViewModelPtr.cpp b/Engine/Source/Editor/SequencerCore/Private/MVVM/ViewModelPtr.cpp
--- a/Engine/Source/Editor/SequencerCore/Private/MVVM/ViewModelPtr.cpp
+++ b/Engine/Source/Editor/SequencerCore/Private/MVVM/ViewModelPtr.cpp
@@ -13,6 +13,10 @@ TViewModelConversions<FViewModel>::operator TWeakPtr<FViewModel>() const
 {
 	return static_cast<const TViewModelPtr<FViewModel>*>(this)->Storage.GetModel();
 }
+TWeakViewModelPtr<FViewModel> TViewModelConversions<FViewModel>::AsWeak() const
+{
+	return *static_cast<const TViewModelPtr<FViewModel>*>(this);
+}
 
 TViewModelConversions<const FViewModel>::operator TSharedPtr<const FViewModel>() const
 {
@@ -22,6 +26,10 @@ TViewModelConversions<const FViewModel>::operator TWeakPtr<const FViewModel>() c
 {
 	return static_cast<const TViewModelPtr<const FViewModel>*>(this)->Storage.GetModel();
 }
+TWeakViewModelPtr<const FViewModel> TViewModelConversions<const FViewModel>::AsWeak() const
+{
+	return *static_cast<const TViewModelPtr<const FViewModel>*>(this);
+}
 
 } // namespace UE::Sequencer
 
diff --git a/Engine/Source/Editor/SequencerCore/Public/MVVM/ViewModelPtr.h b/Engine/Source/Editor/SequencerCore/Public/MVVM/ViewModelPtr.h
--- a/Engine/Source/Editor/SequencerCore/Public/MVVM/ViewModelPtr.h
+++ b/Engine/Source/Editor/SequencerCore/Public/MVVM/ViewModelPtr.h
@@ -70,21 +70,29 @@ struct TViewModelConversions
 		return TSharedPtr<ViewModelType>(*this);
 	}
 	operator TWeakPtr<T>() const
-{
+	{
 		return TSharedPtr<T>(*this);
 	}
+	TWeakViewModelPtr<T> AsWeak() const
+	{
+		return *static_cast<const TViewModelPtr<T>*>(this);
+	}
 };
 template<>
 struct TViewModelConversions<FViewModel>
 {
 	SEQUENCERCORE_API operator TSharedPtr<FViewModel>() const;
 	SEQUENCERCORE_API operator TWeakPtr<FViewModel>() const;
+
+	SEQUENCERCORE_API TWeakViewModelPtr<FViewModel> AsWeak() const;
 };
 template<>
 struct TViewModelConversions<const FViewModel>
 {
 	SEQUENCERCORE_API operator TSharedPtr<const FViewModel>() const;
 	SEQUENCERCORE_API operator TWeakPtr<const FViewModel>() const;
+
+	SEQUENCERCORE_API TWeakViewModelPtr<const FViewModel> AsWeak() const;
 };
 
 /**
@@ -106,6 +114,11 @@ struct TViewModelPtr : TViewModelConversions<T>
 	TViewModelPtr(TViewModelPtr&&) = default;
 	TViewModelPtr& operator=(TViewModelPtr&&) = default;
 
+	explicit TViewModelPtr(T* InViewModelType)
+	{
+		Storage.Set(InViewModelType->AsShared(), InViewModelType);
+	}
+
 	TViewModelPtr(TSharedPtr<ViewModelType> InModel, T* InExtension)
 	{
 		if (InExtension)
