SequencerCore: Added a utility method for creating an outliner button with a custom icon

These behave in exactly the same way as the [+] buttons, but allow for different icons

diff --git a/Engine/Source/Editor/SequencerCore/Private/MVVM/Views/ViewUtilities.cpp b/Engine/Source/Editor/SequencerCore/Private/MVVM/Views/ViewUtilities.cpp
--- a/Engine/Source/Editor/SequencerCore/Private/MVVM/Views/ViewUtilities.cpp
+++ b/Engine/Source/Editor/SequencerCore/Private/MVVM/Views/ViewUtilities.cpp
@@ -14,7 +14,7 @@
 namespace UE::Sequencer
 {
 
-TSharedRef<SWidget> MakeAddButton(FText HoverText, const FOnClicked& HandleClicked, const FViewModelPtr& ViewModel)
+TSharedRef<SWidget> MakeButton(FText HoverText, const FSlateBrush* Image, const FOnClicked& HandleClicked, const FViewModelPtr& ViewModel)
 {
 	TAttribute<bool> IsHovered;
 	if (TSharedPtr<IHoveredExtension> Hoverable = ViewModel.ImplicitCast())
@@ -37,10 +37,10 @@ TSharedRef<SWidget> MakeAddButton(FText HoverText, const FOnClicked& HandleClick
 		.OnClicked(HandleClicked)
 		.IsEnabled(IsEnabled)
 		.IsRowHovered(IsHovered)
-		.Image(FAppStyle::GetBrush("Sequencer.Outliner.Plus"));
+		.Image(Image);
 }
 
-TSharedRef<SWidget> MakeAddButton(FText HoverText, const FOnGetContent& HandleGetMenuContent, const FViewModelPtr& ViewModel)
+TSharedRef<SWidget> MakeButton(FText HoverText, const FSlateBrush* Image, const FOnGetContent& HandleGetMenuContent, const FViewModelPtr& ViewModel)
 {
 	TAttribute<bool> IsHovered;
 	if (TSharedPtr<IHoveredExtension> Hoverable = ViewModel.ImplicitCast())
@@ -63,7 +63,17 @@ TSharedRef<SWidget> MakeAddButton(FText HoverText, const FOnGetContent& HandleGe
 		.OnGetMenuContent(HandleGetMenuContent)
 		.IsEnabled(IsEnabled)
 		.IsRowHovered(IsHovered)
-		.Image(FAppStyle::GetBrush("Sequencer.Outliner.Plus"));
+		.Image(Image);
+}
+
+TSharedRef<SWidget> MakeAddButton(FText HoverText, const FOnClicked& HandleClicked, const FViewModelPtr& ViewModel)
+{
+	return MakeButton(HoverText, FAppStyle::GetBrush("Sequencer.Outliner.Plus"), HandleClicked, ViewModel);
+}
+
+TSharedRef<SWidget> MakeAddButton(FText HoverText, const FOnGetContent& HandleGetMenuContent, const FViewModelPtr& ViewModel)
+{
+	return MakeButton(HoverText, FAppStyle::GetBrush("Sequencer.Outliner.Plus"), HandleGetMenuContent, ViewModel);
 }
 
 TSharedRef<SWidget> MakeAddButton(FText HoverText, const FOnGetContent& MenuContent, const TAttribute<bool>& HoverState, const TAttribute<bool>& IsEnabled)
diff --git a/Engine/Source/Editor/SequencerCore/Public/MVVM/Views/ViewUtilities.h b/Engine/Source/Editor/SequencerCore/Public/MVVM/Views/ViewUtilities.h
--- a/Engine/Source/Editor/SequencerCore/Public/MVVM/Views/ViewUtilities.h
+++ b/Engine/Source/Editor/SequencerCore/Public/MVVM/Views/ViewUtilities.h
@@ -8,6 +8,7 @@
 #include "Templates/SharedPointer.h"
 
 
+struct FSlateBrush;
 class FText;
 class SWidget;
 template< typename ObjectType > class TAttribute;
@@ -15,7 +16,9 @@ template< typename ObjectType > class TAttribute;
 namespace UE::Sequencer
 {
 
-class IOutlinerExtension;
+SEQUENCERCORE_API TSharedRef<SWidget> MakeButton(FText HoverText, const FSlateBrush* Image, const FOnClicked& HandleClicked, const FViewModelPtr& ViewModel);
+
+SEQUENCERCORE_API TSharedRef<SWidget> MakeButton(FText HoverText, const FSlateBrush* Image, const FOnGetContent& HandleGetMenuContent, const FViewModelPtr& ViewModel);
 
 SEQUENCERCORE_API TSharedRef<SWidget> MakeAddButton(FText HoverText, const FOnClicked& HandleClicked, const FViewModelPtr& ViewModel);
 
