Sequencer: Fix crash in OnMouseLeave with null track area view model

#rnx

diff --git a/Engine/Source/Editor/Sequencer/Private/SSequencerSection.cpp b/Engine/Source/Editor/Sequencer/Private/SSequencerSection.cpp
--- a/Engine/Source/Editor/Sequencer/Private/SSequencerSection.cpp
+++ b/Engine/Source/Editor/Sequencer/Private/SSequencerSection.cpp
@@ -2028,7 +2028,10 @@ void SSequencerSection::OnMouseLeave( const FPointerEvent& MouseEvent )
 	SCompoundWidget::OnMouseLeave( MouseEvent );
 
 	TSharedPtr<FTrackAreaViewModel> TrackAreaViewModel = GetTrackAreaViewModel();
-	TrackAreaViewModel->SetHotspot(nullptr);
+	if (TrackAreaViewModel)
+	{
+		TrackAreaViewModel->SetHotspot(nullptr);
+	}
 }
 
 static float SectionThrobDurationSeconds = 1.f;
diff --git a/Engine/Source/Editor/SequencerCore/Private/MVVM/Views/STrackAreaView.cpp b/Engine/Source/Editor/SequencerCore/Private/MVVM/Views/STrackAreaView.cpp
--- a/Engine/Source/Editor/SequencerCore/Private/MVVM/Views/STrackAreaView.cpp
+++ b/Engine/Source/Editor/SequencerCore/Private/MVVM/Views/STrackAreaView.cpp
@@ -377,7 +377,11 @@ void STrackAreaView::UpdateHoverStates( const FGeometry& MyGeometry, const FPoin
 
 	// Set the node that we are hovering
 	TViewModelPtr<IOutlinerExtension> NewHoveredItem = PinnedOutliner->HitTestNode(MyGeometry.AbsoluteToLocal(MouseEvent.GetScreenSpacePosition()).Y);
-	PinnedOutliner->GetOutlinerModel()->SetHoveredItem(NewHoveredItem);
+	TSharedPtr<FOutlinerViewModel> OutlinerModel = PinnedOutliner->GetOutlinerModel();
+	if (OutlinerModel)
+	{
+		OutlinerModel->SetHoveredItem(NewHoveredItem);
+	}
 
 	TSharedPtr<ITrackAreaHotspot> Hotspot = ViewModel->GetHotspot();
 	if (Hotspot.IsValid())
@@ -458,7 +462,11 @@ void STrackAreaView::OnMouseLeave(const FPointerEvent& MouseEvent)
 {
 	if (TSharedPtr<SOutlinerView> Outliner = WeakOutliner.Pin())
 	{
-		Outliner->GetOutlinerModel()->SetHoveredItem(nullptr);
+		TSharedPtr<FOutlinerViewModel> OutlinerModel = Outliner->GetOutlinerModel();
+		if (OutlinerModel)
+		{
+			OutlinerModel->SetHoveredItem(nullptr);
+		}
 	}
 
 	if (TSharedPtr<FTrackAreaViewModel> TrackArea = WeakViewModel.Pin())
