Sequencer: fix half-unused invalidation flag for outliner view models' root items

diff --git a/Engine/Source/Editor/SequencerCore/Private/MVVM/ViewModels/OutlinerViewModel.cpp b/Engine/Source/Editor/SequencerCore/Private/MVVM/ViewModels/OutlinerViewModel.cpp
--- a/Engine/Source/Editor/SequencerCore/Private/MVVM/ViewModels/OutlinerViewModel.cpp
+++ b/Engine/Source/Editor/SequencerCore/Private/MVVM/ViewModels/OutlinerViewModel.cpp
@@ -115,6 +115,8 @@ TArrayView<const TWeakViewModelPtr<IOutlinerExtension>> FOutlinerViewModel::GetT
 	// Update the root items if they have been invalidated
 	if (bRootItemsInvalidated)
 	{
+		bRootItemsInvalidated = false;
+
 		FOutlinerViewModel* MutableThis = const_cast<FOutlinerViewModel*>(this);
 
 		MutableThis->RootItems.Reset();
diff --git a/Engine/Source/Editor/SequencerCore/Public/MVVM/ViewModels/OutlinerViewModel.h b/Engine/Source/Editor/SequencerCore/Public/MVVM/ViewModels/OutlinerViewModel.h
--- a/Engine/Source/Editor/SequencerCore/Public/MVVM/ViewModels/OutlinerViewModel.h
+++ b/Engine/Source/Editor/SequencerCore/Public/MVVM/ViewModels/OutlinerViewModel.h
@@ -108,7 +108,7 @@ private:
 	TWeakViewModelPtr<IOutlinerExtension> WeakHoveredItem;
 
 	/** Whether root items should be re-evaluated */
-	bool bRootItemsInvalidated;
+	mutable bool bRootItemsInvalidated;
 };
 
 } // namespace UE::Sequencer
