Sequencer: fix outliner view model

#rnx

diff --git a/Engine/Source/Editor/SequencerCore/Private/MVVM/ViewModels/OutlinerViewModel.cpp b/Engine/Source/Editor/SequencerCore/Private/MVVM/ViewModels/OutlinerViewModel.cpp
--- a/Engine/Source/Editor/SequencerCore/Private/MVVM/ViewModels/OutlinerViewModel.cpp
+++ b/Engine/Source/Editor/SequencerCore/Private/MVVM/ViewModels/OutlinerViewModel.cpp
@@ -115,10 +115,9 @@ TArrayView<const TWeakViewModelPtr<IOutlinerExtension>> FOutlinerViewModel::GetT
 	// Update the root items if they have been invalidated
 	if (bRootItemsInvalidated)
 	{
-		bRootItemsInvalidated = false;
-
 		FOutlinerViewModel* MutableThis = const_cast<FOutlinerViewModel*>(this);
 
+		MutableThis->bRootItemsInvalidated = false;
 		MutableThis->RootItems.Reset();
 
 		TSharedPtr<FViewModel> Root = WeakRootDataModel.Pin();
diff --git a/Engine/Source/Editor/SequencerCore/Public/MVVM/ViewModels/OutlinerViewModel.h b/Engine/Source/Editor/SequencerCore/Public/MVVM/ViewModels/OutlinerViewModel.h
--- a/Engine/Source/Editor/SequencerCore/Public/MVVM/ViewModels/OutlinerViewModel.h
+++ b/Engine/Source/Editor/SequencerCore/Public/MVVM/ViewModels/OutlinerViewModel.h
@@ -108,7 +108,7 @@ private:
 	TWeakViewModelPtr<IOutlinerExtension> WeakHoveredItem;
 
 	/** Whether root items should be re-evaluated */
-	mutable bool bRootItemsInvalidated;
+	bool bRootItemsInvalidated;
 };
 
 } // namespace UE::Sequencer
