Sequencer: Expand/collapse selected items if this item is selected

For example, you can select 3 bindings, click on the expand arrow and all 3 bindings will be expanded.

diff --git a/Engine/Source/Editor/SequencerCore/Private/MVVM/Views/SOutlinerView.cpp b/Engine/Source/Editor/SequencerCore/Private/MVVM/Views/SOutlinerView.cpp
--- a/Engine/Source/Editor/SequencerCore/Private/MVVM/Views/SOutlinerView.cpp
+++ b/Engine/Source/Editor/SequencerCore/Private/MVVM/Views/SOutlinerView.cpp
@@ -1088,6 +1088,20 @@ void SOutlinerView::OnExpansionChanged(TWeakViewModelPtr<IOutlinerExtension> InW
 			SetItemExpansion(Child, true);
 		}
 	}
+
+	// Also, expand/collapse any selected items if this item is selected
+	TArray<TWeakViewModelPtr<IOutlinerExtension>> WeakSelectedItems = GetSelectedItems();
+	if (WeakSelectedItems.Contains(InWeakItem))
+	{
+		for (TWeakViewModelPtr<IOutlinerExtension> WeakItem : WeakSelectedItems)
+		{
+			TViewModelPtr<IOutlinerExtension> Item = WeakItem.Pin();
+			if (Item && Item != InWeakItem)
+			{
+				SetItemExpansion(Item, bIsExpanded);
+			}
+		}
+	}
 }
 
 void SOutlinerView::SetItemExpansionRecursive(TWeakViewModelPtr<IOutlinerExtension> InItem, bool bIsExpanded)
