Fixed non-unity mac only build failure

diff --git a/Engine/Source/Editor/SequencerCore/Private/SequencerCoreModule.cpp b/Engine/Source/Editor/SequencerCore/Private/SequencerCoreModule.cpp
--- a/Engine/Source/Editor/SequencerCore/Private/SequencerCoreModule.cpp
+++ b/Engine/Source/Editor/SequencerCore/Private/SequencerCoreModule.cpp
@@ -19,7 +19,7 @@ namespace Sequencer
 class FSequencerCoreModule
 	: public ISequencerCoreModule
 {
-	virtual FDelegateHandle RegisterModelType(TWeakObjectPtr<UClass> WeakClass, TFunction<TSharedPtr<FViewModel>(UObject*)> FactoryFunctor) override
+	virtual FDelegateHandle RegisterModelType(const TWeakObjectPtr<UClass>& WeakClass, TFunction<TSharedPtr<FViewModel>(UObject*)> FactoryFunctor) override
 	{
 		TObjectKey<UClass> Key = WeakClass.Get();
 		ensureMsgf(!ModelFactories.Contains(Key), TEXT("Attempting to register a duplicate model factory type"));
diff --git a/Engine/Source/Editor/SequencerCore/Public/ISequencerCoreModule.h b/Engine/Source/Editor/SequencerCore/Public/ISequencerCoreModule.h
--- a/Engine/Source/Editor/SequencerCore/Public/ISequencerCoreModule.h
+++ b/Engine/Source/Editor/SequencerCore/Public/ISequencerCoreModule.h
@@ -30,7 +30,7 @@ public:
 	 * 
 	 * @return A delegate handle that can be passed to UnRegisterModelType to unregister this factory
 	 */
-	virtual FDelegateHandle RegisterModelType(TWeakObjectPtr<UClass> WeakClass, TFunction<TSharedPtr<FViewModel>(UObject*)> FactoryFunctor) = 0;
+	virtual FDelegateHandle RegisterModelType(const TWeakObjectPtr<UClass>& WeakClass, TFunction<TSharedPtr<FViewModel>(UObject*)> FactoryFunctor) = 0;
 
 
 	/**
@@ -60,7 +60,7 @@ public:
 	 * @return A delegate handle that can be passed to UnRegisterModelType to unregister this factory
 	 */
 	template<typename T>
-	FDelegateHandle RegisterModelType(TWeakObjectPtr<UClass> WeakClass, TViewModelTypeID<T> Type)
+	FDelegateHandle RegisterModelType(const TWeakObjectPtr<UClass>& WeakClass, TViewModelTypeID<T> Type)
 	{
 		auto Factory = [this, Type](UObject* Object)
 		{
