[Sequencer]
- Allow SColumnToggleWidget to specify active and inactive brushes
- Add Inactive brush for mute column

diff --git a/Engine/Source/Editor/SequencerCore/Private/MVVM/Views/OutlinerColumns/SColumnToggleWidget.cpp b/Engine/Source/Editor/SequencerCore/Private/MVVM/Views/OutlinerColumns/SColumnToggleWidget.cpp
--- a/Engine/Source/Editor/SequencerCore/Private/MVVM/Views/OutlinerColumns/SColumnToggleWidget.cpp
+++ b/Engine/Source/Editor/SequencerCore/Private/MVVM/Views/OutlinerColumns/SColumnToggleWidget.cpp
@@ -78,6 +78,8 @@ void SColumnToggleWidget::Construct(
 	bIsImplicitlyActive = false;
 
 	ActiveBrush = GetActiveBrush();
+	InactiveBrush = GetInactiveBrush();
+
 	SetCanTick(true);
 
 	static const FName NAME_ChildActiveBrush = TEXT("Sequencer.Column.CheckBoxIndeterminate");
@@ -149,16 +151,25 @@ FSlateColor SColumnToggleWidget::GetImageColorAndOpacity() const
 
 const FSlateBrush* SColumnToggleWidget::GetBrush() const
 {
-	if (bIsChildActive
-		&& !bIsActive
-		&& !bIsMouseOverWidget)
-	{
-		return ChildActiveBrush;
-	}
-	else
+	if (!bIsActive && !bIsImplicitlyActive)
 	{
-		return ActiveBrush;
+		if (bIsChildActive)
+		{
+			if (bIsMouseOverWidget)
+			{
+				return ActiveBrush;
+			}
+
+			return ChildActiveBrush;
+		}
+
+		if (InactiveBrush)
+		{
+			return InactiveBrush;
+		}
 	}
+
+	return ActiveBrush;
 }
 
 FReply SColumnToggleWidget::OnDragDetected(const FGeometry& MyGeometry, const FPointerEvent& MouseEvent)
diff --git a/Engine/Source/Editor/SequencerCore/Private/MVVM/Views/OutlinerColumns/SMuteColumnWidget.cpp b/Engine/Source/Editor/SequencerCore/Private/MVVM/Views/OutlinerColumns/SMuteColumnWidget.cpp
--- a/Engine/Source/Editor/SequencerCore/Private/MVVM/Views/OutlinerColumns/SMuteColumnWidget.cpp
+++ b/Engine/Source/Editor/SequencerCore/Private/MVVM/Views/OutlinerColumns/SMuteColumnWidget.cpp
@@ -113,6 +113,12 @@ const FSlateBrush* SMuteColumnWidget::GetActiveBrush() const
 	return FAppStyle::Get().GetBrush(NAME_MutedBrush);
 }
 
+const FSlateBrush* SMuteColumnWidget::GetInactiveBrush() const
+{
+	static const FName NAME_NotMutedBrush = TEXT("Level.VisibleIcon16x");
+	return FAppStyle::Get().GetBrush(NAME_NotMutedBrush);
+}
+
 #undef LOCTEXT_NAMESPACE
 
 } // namespace UE::Sequencer
diff --git a/Engine/Source/Editor/SequencerCore/Private/MVVM/Views/OutlinerColumns/SMuteColumnWidget.h b/Engine/Source/Editor/SequencerCore/Private/MVVM/Views/OutlinerColumns/SMuteColumnWidget.h
--- a/Engine/Source/Editor/SequencerCore/Private/MVVM/Views/OutlinerColumns/SMuteColumnWidget.h
+++ b/Engine/Source/Editor/SequencerCore/Private/MVVM/Views/OutlinerColumns/SMuteColumnWidget.h
@@ -42,6 +42,9 @@ protected:
 	/** Returns the brush to display when this item is muted. */
 	virtual const FSlateBrush* GetActiveBrush() const override;
 
+	/** Returns the brush to display when this item is not muted. */
+	virtual const FSlateBrush* GetInactiveBrush() const override;
+
 private:
 
 	void SetIsActive_Internal(const FViewModelPtr& ViewModel, const bool bInIsActive);
diff --git a/Engine/Source/Editor/SequencerCore/Public/MVVM/Views/OutlinerColumns/SColumnToggleWidget.h b/Engine/Source/Editor/SequencerCore/Public/MVVM/Views/OutlinerColumns/SColumnToggleWidget.h
--- a/Engine/Source/Editor/SequencerCore/Public/MVVM/Views/OutlinerColumns/SColumnToggleWidget.h
+++ b/Engine/Source/Editor/SequencerCore/Public/MVVM/Views/OutlinerColumns/SColumnToggleWidget.h
@@ -65,6 +65,9 @@ protected:
 	/** Returns the brush to be used to represent a widget is active. */
 	virtual const FSlateBrush* GetActiveBrush() const = 0;
 
+	/** Returns the brush to be used to represent a widget is inactive. If Null, will use the active brush. */
+	virtual const FSlateBrush* GetInactiveBrush() const { return nullptr; }
+
 protected:
 	
 
@@ -121,8 +124,9 @@ protected:
 private:
 
 	/** Brushes for the various states of active widgets. */
-	const FSlateBrush* ActiveBrush;
-	const FSlateBrush* ChildActiveBrush;
+	const FSlateBrush* ActiveBrush = nullptr;
+	const FSlateBrush* InactiveBrush = nullptr;
+	const FSlateBrush* ChildActiveBrush = nullptr;
 
 	/** Scoped undo transaction. */
 	TUniquePtr<FScopedTransaction> UndoTransaction;
