[Sequencer] Remove commented out code block

diff --git a/Engine/Source/Editor/SequencerCore/Private/MVVM/Extensions/IDeactivatableExtension.cpp b/Engine/Source/Editor/SequencerCore/Private/MVVM/Extensions/IDeactivatableExtension.cpp
--- a/Engine/Source/Editor/SequencerCore/Private/MVVM/Extensions/IDeactivatableExtension.cpp
+++ b/Engine/Source/Editor/SequencerCore/Private/MVVM/Extensions/IDeactivatableExtension.cpp
@@ -1,4 +1,4 @@
-﻿// Copyright Epic Games, Inc. All Rights Reserved.
+// Copyright Epic Games, Inc. All Rights Reserved.
 
 #include "MVVM/Extensions/IDeactivatableExtension.h"
 
@@ -42,44 +42,6 @@ ECachedDeactiveState FDeactiveStateCacheExtension::ComputeFlagsForModel(const FV
 	}
 
 	return ThisModelFlags;
-
-	/*ECachedDeactiveState ThisModelFlags = ECachedDeactiveState::None;
-	
-	ECachedDeactiveState ParentFlags = IndividualItemFlags.Last();
-
-	if (EnumHasAnyFlags(ParentFlags, ECachedDeactiveState::Deactivated | ECachedDeactiveState::ImplicitlyDeactivatedByParent))
-	{
-		ThisModelFlags |= ECachedDeactiveState::ImplicitlyDeactivatedByParent;
-	}
-
-	if (const TViewModelPtr<IDeactivatableExtension> Deactivatable = ViewModel.ImplicitCast())
-	{
-		ThisModelFlags |= ECachedDeactiveState::Deactivatable;
-
-		switch (Deactivatable->IsDeactivated())
-		{
-		case EDeactivatableState::Deactivated:
-			ThisModelFlags |= ECachedDeactiveState::Deactivated;
-			break;
-		case EDeactivatableState::PartiallyDeactivated:
-			// Deactivatable, but not deactivated - any parent must not be fully deactivated
-			if (EnumHasAnyFlags(ParentFlags, ECachedDeactiveState::Deactivated))
-			{
-				ParentFlags |= ECachedDeactiveState::PartiallyDeactivatedChildren;
-				ParentFlags &= ~ECachedDeactiveState::Deactivated;
-			}
-			break;
-		case EDeactivatableState::None:
-			break;
-		}
-
-		if (Deactivatable->IsInheritable())
-		{
-			ThisModelFlags |= ECachedDeactiveState::Inheritable;
-		}
-	}
-
-	return ThisModelFlags;*/
 }
 
 void FDeactiveStateCacheExtension::PostComputeChildrenFlags(const FViewModelPtr& ViewModel, ECachedDeactiveState& OutThisModelFlags, ECachedDeactiveState& OutPropagateToParentFlags)
