[Sequencer] Add dash icon for bindings with no child tracks

diff --git a/Engine/Source/Editor/Sequencer/Private/MVVM/ViewModels/OutlinerItemModel.cpp b/Engine/Source/Editor/Sequencer/Private/MVVM/ViewModels/OutlinerItemModel.cpp
--- a/Engine/Source/Editor/Sequencer/Private/MVVM/ViewModels/OutlinerItemModel.cpp
+++ b/Engine/Source/Editor/Sequencer/Private/MVVM/ViewModels/OutlinerItemModel.cpp
@@ -1278,6 +1278,8 @@ bool FEvaluableOutlinerItemModel::IsDeactivated() const
 		return false;
 	}
 
+	bool bNoTrackAreaModels = true;
+
 	for (const TViewModelPtr<ITrackExtension>& TrackNode : Descendants)
 	{
 		const UMovieSceneTrack* const Track = TrackNode->GetTrack();
@@ -1286,6 +1288,14 @@ bool FEvaluableOutlinerItemModel::IsDeactivated() const
 			continue;
 		}
 
+		const TViewModelPtr<ITrackAreaExtension> TrackAreaModel = TrackNode.ImplicitCast();
+		if (!TrackAreaModel->GetTrackAreaModelList())
+		{
+			continue;
+		}
+
+		bNoTrackAreaModels = false;
+
 		if (const TViewModelPtr<FTrackRowModel> TrackRowModel = TrackNode.ImplicitCast())
 		{
 			if (!Track->IsRowEvalDisabled(TrackNode->GetRowIndex(), /*bInCheckLocal=*/false))
@@ -1302,7 +1312,7 @@ bool FEvaluableOutlinerItemModel::IsDeactivated() const
 		}
 	}
 
-	return true;
+	return !bNoTrackAreaModels;
 }
 
 void FEvaluableOutlinerItemModel::SetIsDeactivated(const bool bInIsDeactivated)
diff --git a/Engine/Source/Editor/SequencerCore/Private/MVVM/Views/OutlinerColumns/SDeactivateColumnWidget.cpp b/Engine/Source/Editor/SequencerCore/Private/MVVM/Views/OutlinerColumns/SDeactivateColumnWidget.cpp
--- a/Engine/Source/Editor/SequencerCore/Private/MVVM/Views/OutlinerColumns/SDeactivateColumnWidget.cpp
+++ b/Engine/Source/Editor/SequencerCore/Private/MVVM/Views/OutlinerColumns/SDeactivateColumnWidget.cpp
@@ -5,6 +5,7 @@
 #include "MVVM/Selection/SequencerOutlinerSelection.h"
 #include "MVVM/SharedViewModelData.h"
 #include "MVVM/Extensions/IDeactivatableExtension.h"
+#include "MVVM/Extensions/ITrackAreaExtension.h"
 #include "MVVM/ViewModels/EditorViewModel.h"
 #include "MVVM/ViewModels/OutlinerColumns/IOutlinerColumn.h"
 #include "MVVM/ViewModels/ViewModelIterators.h"
@@ -108,6 +109,33 @@ bool SDeactivateColumnWidget::IsImplicitlyActive() const
 
 const FSlateBrush* SDeactivateColumnWidget::GetActiveBrush() const
 {
+	if (const TViewModelPtr<IOutlinerExtension> OutlinerItem = WeakOutlinerExtension.Pin())
+	{
+		bool bHasTrackAreaModels = false;
+
+		if (const TViewModelPtr<ITrackAreaExtension> TrackArea = OutlinerItem.ImplicitCast())
+		{
+			if (const auto DeactivatableTrackArea = TrackArea->GetTrackAreaModelList())
+			{
+				bHasTrackAreaModels = true;
+			}
+		}
+
+		if (!bHasTrackAreaModels)
+		{
+			if (const TViewModelPtr<IDeactivatableExtension> Deactivatable = OutlinerItem.ImplicitCast())
+			{
+				const TParentFirstChildIterator<IDeactivatableExtension> DeactivatableDescendants
+					= OutlinerItem.AsModel()->GetDescendantsOfType<IDeactivatableExtension>();
+				if (!DeactivatableDescendants)
+				{
+					static const FName PartialBrushName = TEXT("Sequencer.Column.CheckBoxIndeterminate");
+					return FAppStyle::Get().GetBrush(PartialBrushName);
+				}
+			}
+		}
+	}
+
 	static const FName BrushName = TEXT("Sequencer.Column.Mute");
 	return FAppStyle::Get().GetBrush(BrushName);
 }
