Don't use UE_NO_UNIQUE_ADDRESS for castable registration on MSVC because of bad codegen.

diff --git a/Engine/Source/Editor/SequencerCore/Public/MVVM/ICastable.h b/Engine/Source/Editor/SequencerCore/Public/MVVM/ICastable.h
--- a/Engine/Source/Editor/SequencerCore/Public/MVVM/ICastable.h
+++ b/Engine/Source/Editor/SequencerCore/Public/MVVM/ICastable.h
@@ -130,13 +130,23 @@ struct TImplements
 {
 };
 
+
+// MSVC has a compiler bug with zero-sized members of derived types that causes bad code generation for castable types
+// So don't use UE_NO_UNIQUE_ADDRESS on MSVC
+// https://developercommunity.visualstudio.com/t/Bad-code-generation-with-empty-first-mem/10900015
+#ifdef _MSC_VER
+	#define UE_SEQUENCER_NO_UNIQUE_ADDRESS 
+#else
+	#define UE_SEQUENCER_NO_UNIQUE_ADDRESS UE_NO_UNIQUE_ADDRESS
+#endif
+
 #define UE_SEQUENCER_DECLARE_CASTABLE(ThisType, ...)																																\
 	UE_SEQUENCER_DECLARE_CASTABLE_API(, ThisType, __VA_ARGS__)
 
 #define UE_SEQUENCER_DECLARE_CASTABLE_API(MODULE_API, ThisType, ...)																													\
 	UE_SEQUENCER_DECLARE_VIEW_MODEL_TYPE_ID_API(MODULE_API, ThisType);																													\
 	using Implements = ::UE::Sequencer::TImplements<__VA_ARGS__>;																														\
-	UE_NO_UNIQUE_ADDRESS ::UE::Sequencer::ICastable::FAutoRegisterTypeTable AutoRegisterTypeTable = ::UE::Sequencer::ICastable::FAutoRegisterTypeTable(static_cast<ThisType*>(this));	\
+	UE_SEQUENCER_NO_UNIQUE_ADDRESS ::UE::Sequencer::ICastable::FAutoRegisterTypeTable AutoRegisterTypeTable = ::UE::Sequencer::ICastable::FAutoRegisterTypeTable(static_cast<ThisType*>(this));	\
 
 #define UE_SEQUENCER_DEFINE_CASTABLE(ThisType)																																		\
 	UE_SEQUENCER_DEFINE_VIEW_MODEL_TYPE_ID(ThisType);
