Updated headers using UnrealCodeFixup to make sure dllstorage is on methods/staticvars instead of types. Part 5/n

diff --git a/Engine/Source/Editor/SequencerCore/Public/MVVM/Extensions/ITrackLaneExtension.h b/Engine/Source/Editor/SequencerCore/Public/MVVM/Extensions/ITrackLaneExtension.h
--- a/Engine/Source/Editor/SequencerCore/Public/MVVM/Extensions/ITrackLaneExtension.h
+++ b/Engine/Source/Editor/SequencerCore/Public/MVVM/Extensions/ITrackLaneExtension.h
@@ -10,6 +10,8 @@
 #include "Templates/SharedPointer.h"
 #include "Types/SlateEnums.h"
 
+#define UE_API SEQUENCERCORE_API
+
 class FArrangedWidget;
 class ISequencer;
 class SWidget;
@@ -29,13 +31,13 @@ class FEditorViewModel;
 struct FTrackLaneScreenAlignment;
 struct INonLinearTimeTransform;
 
-struct SEQUENCERCORE_API FTrackLaneVerticalArrangement
+struct FTrackLaneVerticalArrangement
 {
 	float Offset = 0;
 	float Height = 0;
 };
 
-struct SEQUENCERCORE_API FTrackLaneVerticalAlignment
+struct FTrackLaneVerticalAlignment
 {
 	enum class ESizeMode : uint8
 	{
@@ -47,10 +49,10 @@ struct SEQUENCERCORE_API FTrackLaneVerticalAlignment
 	EVerticalAlignment VAlign     = VAlign_Center;
 	ESizeMode          Mode       = ESizeMode::Proportional;
 
-	FTrackLaneVerticalArrangement ArrangeWithin(float LayoutHeight) const;
+	UE_API FTrackLaneVerticalArrangement ArrangeWithin(float LayoutHeight) const;
 };
 
-struct SEQUENCERCORE_API FTrackLaneVirtualAlignment
+struct FTrackLaneVirtualAlignment
 {
 	TRange<FFrameNumber> Range;
 	FTrackLaneVerticalAlignment VerticalAlignment;
@@ -70,12 +72,12 @@ struct SEQUENCERCORE_API FTrackLaneVirtualAlignment
 		return FTrackLaneVirtualAlignment { InRange, { InStretchFactor, 0.f, InVAlign, FTrackLaneVerticalAlignment::ESizeMode::Proportional }, InViewSpaceID };
 	}
 
-	TOptional<FFrameNumber> GetFiniteLength() const;
+	UE_API TOptional<FFrameNumber> GetFiniteLength() const;
 
-	FTrackLaneScreenAlignment ToScreen(const FTimeToPixel& TimeToPixel, const FGeometry& ParentGeometry) const;
+	UE_API FTrackLaneScreenAlignment ToScreen(const FTimeToPixel& TimeToPixel, const FGeometry& ParentGeometry) const;
 };
 
-struct SEQUENCERCORE_API FTrackLaneScreenAlignment
+struct FTrackLaneScreenAlignment
 {
 	FTrackLaneScreenAlignment()
 		: LeftPosPx(0.f)
@@ -98,24 +100,24 @@ struct SEQUENCERCORE_API FTrackLaneScreenAlignment
 		return WidthPx > 0.f;
 	}
 
-	FArrangedWidget ArrangeWidget(TSharedRef<SWidget> InWidget, const FGeometry& ParentGeometry) const;
+	UE_API FArrangedWidget ArrangeWidget(TSharedRef<SWidget> InWidget, const FGeometry& ParentGeometry) const;
 };
 
-struct SEQUENCERCORE_API FArrangedVirtualEntity
+struct FArrangedVirtualEntity
 {
 	TRange<FFrameNumber> Range;
 	float VirtualTop, VirtualBottom;
 };
 
 /** Interface used for laying out track lane widgets in screen space */
-class SEQUENCERCORE_API ITrackLaneWidgetSpace
+class ITrackLaneWidgetSpace
 {
 public:
 	virtual FTimeToPixel GetScreenSpace(const FGuid& ViewSpaceID = FGuid()) const = 0;
 };
 
 /** Base interface for track-area lanes */
-class SEQUENCERCORE_API ITrackLaneWidget
+class ITrackLaneWidget
 {
 public:
 
@@ -138,7 +140,7 @@ public:
 	{
 		return FTrackLaneScreenAlignment();
 	}
-	virtual FTrackLaneScreenAlignment GetAlignment(const ITrackLaneWidgetSpace& ScreenSpace, const FGeometry& InParentGeometry) const;
+	UE_API virtual FTrackLaneScreenAlignment GetAlignment(const ITrackLaneWidgetSpace& ScreenSpace, const FGeometry& InParentGeometry) const;
 
 	/**
 	 * Gets this widget's overlap priority
@@ -162,7 +164,7 @@ public:
 };
 
 /** Parameters for creating a track lane widget */
-struct SEQUENCERCORE_API FCreateTrackLaneViewParams
+struct FCreateTrackLaneViewParams
 {
 	FCreateTrackLaneViewParams(const TSharedPtr<FEditorViewModel> InEditor)
 		: Editor(InEditor)
@@ -178,11 +180,11 @@ struct SEQUENCERCORE_API FCreateTrackLaneViewParams
 };
 
 /** Extension for view-models that can create track lanes in the track area */
-class SEQUENCERCORE_API ITrackLaneExtension
+class ITrackLaneExtension
 {
 public:
 
-	UE_SEQUENCER_DECLARE_VIEW_MODEL_TYPE_ID(ITrackLaneExtension)
+	UE_SEQUENCER_DECLARE_VIEW_MODEL_TYPE_ID_API(UE_API, ITrackLaneExtension)
 
 	virtual ~ITrackLaneExtension(){}
 
@@ -193,3 +195,4 @@ public:
 } // namespace Sequencer
 } // namespace UE
 
+#undef UE_API
diff --git a/Engine/Source/Editor/SequencerCore/Public/MVVM/Extensions/IViewSpaceClientExtension.h b/Engine/Source/Editor/SequencerCore/Public/MVVM/Extensions/IViewSpaceClientExtension.h
--- a/Engine/Source/Editor/SequencerCore/Public/MVVM/Extensions/IViewSpaceClientExtension.h
+++ b/Engine/Source/Editor/SequencerCore/Public/MVVM/Extensions/IViewSpaceClientExtension.h
@@ -7,6 +7,8 @@
 #include "MVVM/ViewModelTypeID.h"
 #include "MVVM/Extensions/DynamicExtensionContainer.h"
 
+#define UE_API SEQUENCERCORE_API
+
 struct FGuid;
 
 namespace UE::Sequencer
@@ -15,11 +17,11 @@ namespace UE::Sequencer
 /**
  * An extension that is used for view models that represent UObject data
  */
-class SEQUENCERCORE_API IViewSpaceClientExtension
+class IViewSpaceClientExtension
 {
 public:
 
-	UE_SEQUENCER_DECLARE_VIEW_MODEL_TYPE_ID(IViewSpaceClientExtension)
+	UE_SEQUENCER_DECLARE_VIEW_MODEL_TYPE_ID_API(UE_API, IViewSpaceClientExtension)
 
 	virtual ~IViewSpaceClientExtension(){}
 
@@ -31,13 +33,13 @@ public:
 /**
  * An extension that is used for view models that represent UObject data
  */
-class SEQUENCERCORE_API FViewSpaceClientExtensionShim
+class FViewSpaceClientExtensionShim
 	: public IDynamicExtension
 	, public IViewSpaceClientExtension
 {
 public:
 
-	UE_SEQUENCER_DECLARE_VIEW_MODEL_TYPE_ID(FViewSpaceClientExtensionShim)
+	UE_SEQUENCER_DECLARE_VIEW_MODEL_TYPE_ID_API(UE_API, FViewSpaceClientExtensionShim)
 
 	using Implements = TImplements<IDynamicExtension, IViewSpaceClientExtension>;
 
@@ -63,4 +65,6 @@ protected:
 
 
 
-} // namespace UE::Sequencer
\ No newline at end of file
+} // namespace UE::Sequencer
+
+#undef UE_API
diff --git a/Engine/Source/Editor/SequencerCore/Public/MVVM/Extensions/LinkedOutlinerExtension.h b/Engine/Source/Editor/SequencerCore/Public/MVVM/Extensions/LinkedOutlinerExtension.h
--- a/Engine/Source/Editor/SequencerCore/Public/MVVM/Extensions/LinkedOutlinerExtension.h
+++ b/Engine/Source/Editor/SequencerCore/Public/MVVM/Extensions/LinkedOutlinerExtension.h
@@ -7,6 +7,8 @@
 #include "MVVM/ViewModelPtr.h"
 #include "MVVM/ViewModelTypeID.h"
 
+#define UE_API SEQUENCERCORE_API
+
 namespace UE
 {
 namespace Sequencer
@@ -19,17 +21,17 @@ class FViewModel;
  * This is the case for example in Sequencer where the channels of the sections in the track area
  * are linked to the channel group items in the outliner.
  */
-class SEQUENCERCORE_API FLinkedOutlinerExtension
+class FLinkedOutlinerExtension
 {
 public:
 
-	UE_SEQUENCER_DECLARE_VIEW_MODEL_TYPE_ID(FLinkedOutlinerExtension);
+	UE_SEQUENCER_DECLARE_VIEW_MODEL_TYPE_ID_API(UE_API, FLinkedOutlinerExtension);
 
-	FLinkedOutlinerExtension();
+	UE_API FLinkedOutlinerExtension();
 
-	TViewModelPtr<IOutlinerExtension> GetLinkedOutlinerItem() const;
+	UE_API TViewModelPtr<IOutlinerExtension> GetLinkedOutlinerItem() const;
 
-	void SetLinkedOutlinerItem(const TViewModelPtr<IOutlinerExtension>& InOutlinerItem);
+	UE_API void SetLinkedOutlinerItem(const TViewModelPtr<IOutlinerExtension>& InOutlinerItem);
 
 private:
 
@@ -41,7 +43,7 @@ private:
  * This is useful for instance when the outliner item is going to compute sizing info for
  * all of its linked items.
  */
-class SEQUENCERCORE_API FLinkedOutlinerComputedSizingShim
+class FLinkedOutlinerComputedSizingShim
 {
 public:
 
@@ -56,3 +58,4 @@ protected:
 } // namespace Sequencer
 } // namespace UE
 
+#undef UE_API
diff --git a/Engine/Source/Editor/SequencerCore/Public/MVVM/Extensions/ViewModelExtensionCollection.h b/Engine/Source/Editor/SequencerCore/Public/MVVM/Extensions/ViewModelExtensionCollection.h
--- a/Engine/Source/Editor/SequencerCore/Public/MVVM/Extensions/ViewModelExtensionCollection.h
+++ b/Engine/Source/Editor/SequencerCore/Public/MVVM/Extensions/ViewModelExtensionCollection.h
@@ -10,6 +10,8 @@
 #include "Templates/SharedPointer.h"
 #include "Templates/UnrealTemplate.h"
 
+#define UE_API SEQUENCERCORE_API
+
 template<typename, typename> class TArrayView;
 
 namespace UE
@@ -25,22 +27,22 @@ struct TViewModelExtensionCollection;
  * Utility class for maintaining a list of view-models that implement a given extension
  * in a given hierarchy. The list is kept up to date when the hierarchy changes.
  */
-struct SEQUENCERCORE_API FViewModelExtensionCollection : FNoncopyable
+struct FViewModelExtensionCollection : FNoncopyable
 {
 private:
 
-	FViewModelExtensionCollection(FViewModelTypeID InExtensionType);
-	explicit FViewModelExtensionCollection(FViewModelTypeID InExtensionType, TWeakPtr<FViewModel> InWeakModel, int32 InDesiredRecursionDepth = -1);
+	UE_API FViewModelExtensionCollection(FViewModelTypeID InExtensionType);
+	UE_API explicit FViewModelExtensionCollection(FViewModelTypeID InExtensionType, TWeakPtr<FViewModel> InWeakModel, int32 InDesiredRecursionDepth = -1);
 
-	virtual ~FViewModelExtensionCollection();
+	UE_API virtual ~FViewModelExtensionCollection();
 
 protected:
 
 	virtual void OnExtensionsDirtied() {}
 
-	void Initialize();
-	void Reinitialize(TWeakPtr<FViewModel> InWeakModel, int32 InDesiredRecursionDepth = -1);
-	void Destroy();
+	UE_API void Initialize();
+	UE_API void Reinitialize(TWeakPtr<FViewModel> InWeakModel, int32 InDesiredRecursionDepth = -1);
+	UE_API void Destroy();
 
 	TSharedPtr<FViewModel> GetObservedModel() const
 	{
@@ -49,13 +51,13 @@ protected:
 
 private:
 
-	void Update() const;
+	UE_API void Update() const;
 
-	void OnHierarchyUpdated();
+	UE_API void OnHierarchyUpdated();
 
-	void ConditionalUpdate() const;
+	UE_API void ConditionalUpdate() const;
 
-	void DestroyImpl();
+	UE_API void DestroyImpl();
 
 	template<typename T>
 	TArrayView<T* const> GetExtensions() const
@@ -119,3 +121,4 @@ struct TViewModelExtensionCollection : FViewModelExtensionCollection
 } // namespace Sequencer
 } // namespace UE
 
+#undef UE_API
diff --git a/Engine/Source/Editor/SequencerCore/Public/MVVM/ICastable.h b/Engine/Source/Editor/SequencerCore/Public/MVVM/ICastable.h
--- a/Engine/Source/Editor/SequencerCore/Public/MVVM/ICastable.h
+++ b/Engine/Source/Editor/SequencerCore/Public/MVVM/ICastable.h
@@ -12,13 +12,15 @@
 
 #include <type_traits>
 
+#define UE_API SEQUENCERCORE_API
+
 namespace UE::Sequencer
 {
 
 struct FCastableTypeTable;
 class FDynamicExtensionContainer;
 
-class SEQUENCERCORE_API ICastable
+class ICastable
 {
 public:
 
@@ -65,8 +67,8 @@ public:
 		return CastThis<T>() != nullptr;
 	}
 
-	void* CastRaw(FViewModelTypeID InType);
-	const void* CastRaw(FViewModelTypeID InType) const;
+	UE_API void* CastRaw(FViewModelTypeID InType);
+	UE_API const void* CastRaw(FViewModelTypeID InType) const;
 
 	template<typename T>
 	T* CastThis()
@@ -152,4 +154,6 @@ struct TImplements
 	UE_SEQUENCER_DEFINE_VIEW_MODEL_TYPE_ID(ThisType);
 
 
-} // namespace UE::Sequencer
\ No newline at end of file
+} // namespace UE::Sequencer
+
+#undef UE_API
diff --git a/Engine/Source/Editor/SequencerCore/Public/MVVM/SharedViewModelData.h b/Engine/Source/Editor/SequencerCore/Public/MVVM/SharedViewModelData.h
--- a/Engine/Source/Editor/SequencerCore/Public/MVVM/SharedViewModelData.h
+++ b/Engine/Source/Editor/SequencerCore/Public/MVVM/SharedViewModelData.h
@@ -11,6 +11,8 @@
 #include "Templates/SharedPointer.h"
 #include "Templates/UniquePtr.h"
 
+#define UE_API SEQUENCERCORE_API
+
 class FDelegateHandle;
 
 namespace UE
@@ -20,23 +22,23 @@ namespace Sequencer
 
 class FViewModel;
 
-class SEQUENCERCORE_API FSharedViewModelData
+class FSharedViewModelData
 	: public FViewModel
 {
 public:
 
-	UE_SEQUENCER_DECLARE_CASTABLE(FSharedViewModelData, FViewModel);
+	UE_SEQUENCER_DECLARE_CASTABLE_API(UE_API, FSharedViewModelData, FViewModel);
 
-	void PreHierarchicalChange(const TSharedPtr<FViewModel>& InChangedModel);
-	void BroadcastHierarchicalChange(const TSharedPtr<FViewModel>& InChangedModel);
+	UE_API void PreHierarchicalChange(const TSharedPtr<FViewModel>& InChangedModel);
+	UE_API void BroadcastHierarchicalChange(const TSharedPtr<FViewModel>& InChangedModel);
 
-	void ReportLatentHierarchicalOperations();
+	UE_API void ReportLatentHierarchicalOperations();
 
-	FSimpleMulticastDelegate& SubscribeToHierarchyChanged(const TSharedPtr<FViewModel>& InModel);
-	void UnsubscribeFromHierarchyChanged(const TSharedPtr<FViewModel>& InModel, FDelegateHandle InHandle);
-	void UnsubscribeFromHierarchyChanged(const TSharedPtr<FViewModel>& InModel, FDelegateUserObjectConst InUserObject);
+	UE_API FSimpleMulticastDelegate& SubscribeToHierarchyChanged(const TSharedPtr<FViewModel>& InModel);
+	UE_API void UnsubscribeFromHierarchyChanged(const TSharedPtr<FViewModel>& InModel, FDelegateHandle InHandle);
+	UE_API void UnsubscribeFromHierarchyChanged(const TSharedPtr<FViewModel>& InModel, FDelegateUserObjectConst InUserObject);
 
-	void PurgeStaleHandlers();
+	UE_API void PurgeStaleHandlers();
 
 private:
 	friend FViewModelHierarchyOperation;
@@ -51,3 +53,4 @@ private:
 } // namespace Sequencer
 } // namespace UE
 
+#undef UE_API
diff --git a/Engine/Source/Editor/SequencerCore/Public/MVVM/ViewModels/EditorSharedViewModelData.h b/Engine/Source/Editor/SequencerCore/Public/MVVM/ViewModels/EditorSharedViewModelData.h
--- a/Engine/Source/Editor/SequencerCore/Public/MVVM/ViewModels/EditorSharedViewModelData.h
+++ b/Engine/Source/Editor/SequencerCore/Public/MVVM/ViewModels/EditorSharedViewModelData.h
@@ -6,15 +6,17 @@
 #include "MVVM/SharedViewModelData.h"
 #include "MVVM/ViewModels/EditorViewModel.h"
 
+#define UE_API SEQUENCERCORE_API
+
 namespace UE::Sequencer
 {
 
-class SEQUENCERCORE_API FEditorSharedViewModelData
+class FEditorSharedViewModelData
 	: public FSharedViewModelData
 {
 public:
 
-	UE_SEQUENCER_DECLARE_CASTABLE(FEditorSharedViewModelData, FSharedViewModelData);
+	UE_SEQUENCER_DECLARE_CASTABLE_API(UE_API, FEditorSharedViewModelData, FSharedViewModelData);
 
 	FEditorSharedViewModelData(TSharedRef<FEditorViewModel> InEditor)
 		: WeakEditor(InEditor)
@@ -34,3 +36,4 @@ private:
 
 } // namespace UE::Sequencer
 
+#undef UE_API
diff --git a/Engine/Source/Editor/SequencerCore/Public/MVVM/ViewModels/EditorViewModel.h b/Engine/Source/Editor/SequencerCore/Public/MVVM/ViewModels/EditorViewModel.h
--- a/Engine/Source/Editor/SequencerCore/Public/MVVM/ViewModels/EditorViewModel.h
+++ b/Engine/Source/Editor/SequencerCore/Public/MVVM/ViewModels/EditorViewModel.h
@@ -14,6 +14,8 @@
 #include "TimeToPixel.h"
 #include "UObject/StrongObjectPtr.h"
 
+#define UE_API SEQUENCERCORE_API
+
 class USequencerScriptingLayer;
 
 namespace UE::Sequencer
@@ -36,50 +38,50 @@ template <typename T> struct TAutoRegisterViewModelTypeID;
  * call the necessary code to refresh the entire view-model hierarchy.
  * Note that this root view-model is *not* parented under the editor view-model.
  */
-class SEQUENCERCORE_API FEditorViewModel
+class FEditorViewModel
 	: public FViewModel
 {
 public:
 
-	UE_SEQUENCER_DECLARE_CASTABLE(FEditorViewModel, FViewModel);
+	UE_SEQUENCER_DECLARE_CASTABLE_API(UE_API, FEditorViewModel, FViewModel);
 
 	/** Children type for editor panels */
-	static EViewModelListType GetEditorPanelType();
+	static UE_API EViewModelListType GetEditorPanelType();
 
 	/** Builds a new editor view model. */
-	FEditorViewModel();
+	UE_API FEditorViewModel();
 	/** Destroys this editor view model. */
-	~FEditorViewModel();
+	UE_API ~FEditorViewModel();
 
 	/** Initializes the outliner and track-area for this editor. */
-	void InitializeEditor();
+	UE_API void InitializeEditor();
 
 	/** Gets the root view-model of the sequence being edited. */
-	FViewModelPtr GetRootModel() const;
+	UE_API FViewModelPtr GetRootModel() const;
 
 	/** Gets the editor panels' view models. */
-	FViewModelChildren GetEditorPanels();
+	UE_API FViewModelChildren GetEditorPanels();
 
 	/** Gets the outliner view-model. */
-	TSharedPtr<FOutlinerViewModel>  GetOutliner() const;
+	UE_API TSharedPtr<FOutlinerViewModel>  GetOutliner() const;
 
 	/** Gets the track area view-model. */
-	TSharedPtr<FTrackAreaViewModel> GetTrackArea() const;
+	UE_API TSharedPtr<FTrackAreaViewModel> GetTrackArea() const;
 
 	/** Gets the selection class */
-	TSharedPtr<FSequencerCoreSelection> GetSelection() const;
+	UE_API TSharedPtr<FSequencerCoreSelection> GetSelection() const;
 
 	/** Gets the scripting layer */
-	USequencerScriptingLayer* GetScriptingLayer() const;
+	UE_API USequencerScriptingLayer* GetScriptingLayer() const;
 
 	/** Returns the current view density for this editor */
-	FViewDensityInfo GetViewDensity() const;
+	UE_API FViewDensityInfo GetViewDensity() const;
 
 	/** Returns the current view density for this editor */
-	void SetViewDensity(const FViewDensityInfo& InViewDensity);
+	UE_API void SetViewDensity(const FViewDensityInfo& InViewDensity);
 
 	/** Returns whether this editor is currently read-only */
-	virtual bool IsReadOnly() const;
+	UE_API virtual bool IsReadOnly() const;
 
 	/** Returns the inverse of read-only - useful for direct bindings to IsEnabled for widgets */
 	bool IsEditable() const
@@ -92,15 +94,15 @@ protected:
 	/** Initializes this view model before panels and root models are created */
 	virtual void PreInitializeEditorImpl() {}
 	/** Creates the root data model for this editor. */
-	virtual TSharedPtr<FViewModel> CreateRootModelImpl();
+	UE_API virtual TSharedPtr<FViewModel> CreateRootModelImpl();
 	/** Creates the outliner for this editor. */
-	virtual TSharedPtr<FOutlinerViewModel> CreateOutlinerImpl();
+	UE_API virtual TSharedPtr<FOutlinerViewModel> CreateOutlinerImpl();
 	/** Creates the track-area for this editor. */
-	virtual TSharedPtr<FTrackAreaViewModel> CreateTrackAreaImpl();
+	UE_API virtual TSharedPtr<FTrackAreaViewModel> CreateTrackAreaImpl();
 	/** Creates the selection class for this editor. */
-	virtual TSharedPtr<FSequencerCoreSelection> CreateSelectionImpl();
+	UE_API virtual TSharedPtr<FSequencerCoreSelection> CreateSelectionImpl();
 	/** Creates the scripting layer class for this editor. */
-	virtual USequencerScriptingLayer* CreateScriptingLayerImpl();
+	UE_API virtual USequencerScriptingLayer* CreateScriptingLayerImpl();
 	/** Creates any other panels for this editor. */
 	virtual void InitializeEditorImpl() {}
 
@@ -129,3 +131,4 @@ private:
 
 } // namespace UE::Sequencer
 
+#undef UE_API
diff --git a/Engine/Source/Editor/SequencerCore/Public/MVVM/ViewModels/OutlinerColumns/IOutlinerColumn.h b/Engine/Source/Editor/SequencerCore/Public/MVVM/ViewModels/OutlinerColumns/IOutlinerColumn.h
--- a/Engine/Source/Editor/SequencerCore/Public/MVVM/ViewModels/OutlinerColumns/IOutlinerColumn.h
+++ b/Engine/Source/Editor/SequencerCore/Public/MVVM/ViewModels/OutlinerColumns/IOutlinerColumn.h
@@ -35,7 +35,7 @@ struct FCreateOutlinerColumnParams
 /**
 * Interface for sequencer outliner columns.
 */
-class SEQUENCERCORE_API IOutlinerColumn : public TSharedFromThis<IOutlinerColumn>
+class IOutlinerColumn : public TSharedFromThis<IOutlinerColumn>
 {
 
 public:
@@ -72,4 +72,4 @@ public:
 };
 
 
-} // namespace UE::Sequencer
\ No newline at end of file
+} // namespace UE::Sequencer
diff --git a/Engine/Source/Editor/SequencerCore/Public/MVVM/ViewModels/OutlinerSpacer.h b/Engine/Source/Editor/SequencerCore/Public/MVVM/ViewModels/OutlinerSpacer.h
--- a/Engine/Source/Editor/SequencerCore/Public/MVVM/ViewModels/OutlinerSpacer.h
+++ b/Engine/Source/Editor/SequencerCore/Public/MVVM/ViewModels/OutlinerSpacer.h
@@ -14,6 +14,8 @@
 #include "UObject/NameTypes.h"
 #include "Widgets/Views/STableRow.h"
 
+#define UE_API SEQUENCERCORE_API
+
 class FDragDropEvent;
 class SWidget;
 namespace UE::Sequencer { template <typename T> struct TAutoRegisterViewModelTypeID; }
@@ -23,7 +25,7 @@ namespace UE
 namespace Sequencer
 {
 
-class SEQUENCERCORE_API FOutlinerSpacer
+class FOutlinerSpacer
 	: public FViewModel
 	, public FGeometryExtensionShim
 	, public FOutlinerExtensionShim
@@ -32,10 +34,10 @@ class SEQUENCERCORE_API FOutlinerSpacer
 {
 public:
 
-	UE_SEQUENCER_DECLARE_CASTABLE(FOutlinerSpacer, FViewModel, IGeometryExtension, IOutlinerExtension, ISortableExtension, IOutlinerDropTargetOutlinerExtension);
+	UE_SEQUENCER_DECLARE_CASTABLE_API(UE_API, FOutlinerSpacer, FViewModel, IGeometryExtension, IOutlinerExtension, ISortableExtension, IOutlinerDropTargetOutlinerExtension);
 
-	FOutlinerSpacer(float InDesiredSpacerHeight);
-	~FOutlinerSpacer();
+	UE_API FOutlinerSpacer(float InDesiredSpacerHeight);
+	UE_API ~FOutlinerSpacer();
 
 	void SetDesiredHeight(float InDesiredSpacerHeight)
 	{
@@ -45,19 +47,19 @@ public:
 public:
 
 	/*~ IOutlinerExtension */
-	bool HasBackground() const override;
-	FName GetIdentifier() const override;
-	FOutlinerSizing GetOutlinerSizing() const override;
-	TSharedPtr<SWidget> CreateContextMenuWidget(const FCreateOutlinerContextMenuWidgetParams& InParams) override;
+	UE_API bool HasBackground() const override;
+	UE_API FName GetIdentifier() const override;
+	UE_API FOutlinerSizing GetOutlinerSizing() const override;
+	UE_API TSharedPtr<SWidget> CreateContextMenuWidget(const FCreateOutlinerContextMenuWidgetParams& InParams) override;
 
 	/*~ ISortableExtension */
-	virtual void SortChildren() override;
-	virtual FSortingKey GetSortingKey() const override;
-	virtual void SetCustomOrder(int32 InCustomOrder) override;
+	UE_API virtual void SortChildren() override;
+	UE_API virtual FSortingKey GetSortingKey() const override;
+	UE_API virtual void SetCustomOrder(int32 InCustomOrder) override;
 
 	/*~ IOutlinerDropTargetOutlinerExtension */
-	TOptional<EItemDropZone> CanAcceptDrop(const FViewModelPtr& TargetModel, const FDragDropEvent& DragDropEvent, EItemDropZone InItemDropZone) override;
-	void PerformDrop(const FViewModelPtr& TargetModel, const FDragDropEvent& DragDropEvent, EItemDropZone InItemDropZone) override;
+	UE_API TOptional<EItemDropZone> CanAcceptDrop(const FViewModelPtr& TargetModel, const FDragDropEvent& DragDropEvent, EItemDropZone InItemDropZone) override;
+	UE_API void PerformDrop(const FViewModelPtr& TargetModel, const FDragDropEvent& DragDropEvent, EItemDropZone InItemDropZone) override;
 
 private:
 
@@ -68,3 +70,4 @@ private:
 } // namespace Sequencer
 } // namespace UE
 
+#undef UE_API
diff --git a/Engine/Source/Editor/SequencerCore/Public/MVVM/ViewModels/OutlinerViewModel.h b/Engine/Source/Editor/SequencerCore/Public/MVVM/ViewModels/OutlinerViewModel.h
--- a/Engine/Source/Editor/SequencerCore/Public/MVVM/ViewModels/OutlinerViewModel.h
+++ b/Engine/Source/Editor/SequencerCore/Public/MVVM/ViewModels/OutlinerViewModel.h
@@ -13,6 +13,8 @@
 #include "MVVM/ViewModels/ViewModel.h"
 #include "Templates/SharedPointer.h"
 
+#define UE_API SEQUENCERCORE_API
+
 class FDragDropOperation;
 class SWidget;
 
@@ -23,19 +25,19 @@ class FEditorViewModel;
 class FViewModel;
 class IOutlinerExtension;
 
-class SEQUENCERCORE_API FOutlinerViewModel
+class FOutlinerViewModel
 	: public FViewModel
 {
 public:
 
-	UE_SEQUENCER_DECLARE_CASTABLE(FOutlinerViewModel, FViewModel);
+	UE_SEQUENCER_DECLARE_CASTABLE_API(UE_API, FOutlinerViewModel, FViewModel);
 
 	FSimpleMulticastDelegate OnRefreshed;
 
 	/**
 	 * Constructor
 	 */
-	FOutlinerViewModel();
+	UE_API FOutlinerViewModel();
 
 	/**
 	 * Destructor
@@ -45,57 +47,57 @@ public:
 	/**
 	 * Initialize the view model with a given root model
 	 */
-	void Initialize(const FWeakViewModelPtr& InWeakRootDataModel);
+	UE_API void Initialize(const FWeakViewModelPtr& InWeakRootDataModel);
 
 	/**
 	 * Get the editor view-model provided by the creation context
 	 */
-	TSharedPtr<FEditorViewModel> GetEditor() const;
+	UE_API TSharedPtr<FEditorViewModel> GetEditor() const;
 
 	/**
 	 * Get the logical root item of the outliner hierarchy
 	 */
-	FViewModelPtr GetRootItem() const;
+	UE_API FViewModelPtr GetRootItem() const;
 
 	/**
 	 * Retrieve a flat list of the root items in this outliner model
 	 */
-	TArrayView<const TWeakViewModelPtr<IOutlinerExtension>> GetTopLevelItems() const;
+	UE_API TArrayView<const TWeakViewModelPtr<IOutlinerExtension>> GetTopLevelItems() const;
 
 	/**
 	 * Set the single hovered model in the tree
 	 */
-	void SetHoveredItem(TViewModelPtr<IOutlinerExtension> InHoveredModel);
+	UE_API void SetHoveredItem(TViewModelPtr<IOutlinerExtension> InHoveredModel);
 
 	/**
 	 * Get the single hovered model in the tree, possibly nullptr
 	 */
-	TViewModelPtr<IOutlinerExtension> GetHoveredItem() const;
+	UE_API TViewModelPtr<IOutlinerExtension> GetHoveredItem() const;
 
 	/**
 	 * Unpins any pinned nodes in this tree
 	 */
-	void UnpinAllNodes();
+	UE_API void UnpinAllNodes();
 
 	/**
 	 * Called when the default context menu is opened
 	 */
-	virtual TSharedPtr<SWidget> CreateContextMenuWidget();
+	UE_API virtual TSharedPtr<SWidget> CreateContextMenuWidget();
 
 	/**
 	 * Initate a new drag operation from the specified dragged models.
 	 */
-	virtual TSharedRef<FDragDropOperation> InitiateDrag(TArray<TWeakViewModelPtr<IOutlinerExtension>>&& InDraggedModels);
+	UE_API virtual TSharedRef<FDragDropOperation> InitiateDrag(TArray<TWeakViewModelPtr<IOutlinerExtension>>&& InDraggedModels);
 
 	/**
 	 * Request that this tree be updated
 	 */
-	virtual void RequestUpdate();
+	UE_API virtual void RequestUpdate();
 
 
 private:
 
-	void HandleDataHierarchyChanged();
+	UE_API void HandleDataHierarchyChanged();
 
 private:
 
@@ -113,3 +115,4 @@ private:
 
 } // namespace UE::Sequencer
 
+#undef UE_API
diff --git a/Engine/Source/Editor/SequencerCore/Public/MVVM/ViewModels/OutlinerViewModelDragDropOp.h b/Engine/Source/Editor/SequencerCore/Public/MVVM/ViewModels/OutlinerViewModelDragDropOp.h
--- a/Engine/Source/Editor/SequencerCore/Public/MVVM/ViewModels/OutlinerViewModelDragDropOp.h
+++ b/Engine/Source/Editor/SequencerCore/Public/MVVM/ViewModels/OutlinerViewModelDragDropOp.h
@@ -15,6 +15,8 @@
 #include "Misc/FrameNumber.h"
 #include "Templates/SharedPointer.h"
 
+#define UE_API SEQUENCERCORE_API
+
 struct FSlateBrush;
 
 namespace UE
@@ -26,7 +28,7 @@ class FViewModel;
 class IOutlinerExtension;
 
 /** A decorated drag drop operation object for dragging sequencer outliner items. */
-class SEQUENCERCORE_API FOutlinerViewModelDragDropOp 
+class FOutlinerViewModelDragDropOp 
 	: public FGraphEditorDragDropAction
 	, public IDragOperation
 {
@@ -36,7 +38,7 @@ public:
 	/**
 	 * Construct a new drag/drop operation for dragging a selection of display nodes
 	 */
-	static TSharedRef<FOutlinerViewModelDragDropOp> New(TArray<TWeakViewModelPtr<IOutlinerExtension>>&& InDraggedViewModels, FText InDefaultText, const FSlateBrush* InDefaultIcon);
+	static UE_API TSharedRef<FOutlinerViewModelDragDropOp> New(TArray<TWeakViewModelPtr<IOutlinerExtension>>&& InDraggedViewModels, FText InDefaultText, const FSlateBrush* InDefaultIcon);
 
 protected:
 
@@ -47,25 +49,25 @@ protected:
 public:
 
 	//~ FDragDropOperation interface
-	virtual void Construct() override;
+	UE_API virtual void Construct() override;
 
 public:
 
 	/**
 	 * Reset the tooltip decorator back to its original state
 	 */
-	void ResetToDefaultToolTip();
+	UE_API void ResetToDefaultToolTip();
 
 	/**
 	 * Gets the nodes which are currently being dragged.
 	 */
-	TArrayView<const TWeakViewModelPtr<IOutlinerExtension>> GetDraggedViewModels() const;
+	UE_API TArrayView<const TWeakViewModelPtr<IOutlinerExtension>> GetDraggedViewModels() const;
 
 	/**
 	 * Check whether dropping onto the specified outliner item would result in a parent being dropped onto its child
 	 * @note: Sets this operation's CurrentHoverText to an error string on failure
 	 */
-	bool ValidateParentChildDrop(const FViewModel& ProspectiveItem);
+	UE_API bool ValidateParentChildDrop(const FViewModel& ProspectiveItem);
 
 private:
 
@@ -118,3 +120,4 @@ protected:
 } // namespace Sequencer
 } // namespace UE
 
+#undef UE_API
diff --git a/Engine/Source/Editor/SequencerCore/Public/MVVM/ViewModels/TrackAreaViewModel.h b/Engine/Source/Editor/SequencerCore/Public/MVVM/ViewModels/TrackAreaViewModel.h
--- a/Engine/Source/Editor/SequencerCore/Public/MVVM/ViewModels/TrackAreaViewModel.h
+++ b/Engine/Source/Editor/SequencerCore/Public/MVVM/ViewModels/TrackAreaViewModel.h
@@ -12,6 +12,8 @@
 #include "TimeToPixel.h"
 #include "UObject/NameTypes.h"
 
+#define UE_API SEQUENCERCORE_API
+
 struct FFrameNumber;
 struct FFrameRate;
 struct FGeometry;
@@ -28,15 +30,15 @@ struct ITrackAreaHotspot;
 
 DECLARE_MULTICAST_DELEGATE_OneParam(FOnTrackAreaHotspotChanged, TSharedPtr<ITrackAreaHotspot>);
 
-class SEQUENCERCORE_API FTrackAreaViewModel
+class FTrackAreaViewModel
 	: public FViewModel
 {
 public:
 
-	UE_SEQUENCER_DECLARE_CASTABLE(FTrackAreaViewModel, FViewModel);
+	UE_SEQUENCER_DECLARE_CASTABLE_API(UE_API, FTrackAreaViewModel, FViewModel);
 
-	FTrackAreaViewModel();
-	virtual ~FTrackAreaViewModel();
+	UE_API FTrackAreaViewModel();
+	UE_API virtual ~FTrackAreaViewModel();
 
 public:
 
@@ -46,54 +48,54 @@ public:
 	FViewSpaceEvent OnViewSpaceRemoved;
 
 	/** Get the editor view-model provided by the creation context */
-	TSharedPtr<FEditorViewModel> GetEditor() const;
+	UE_API TSharedPtr<FEditorViewModel> GetEditor() const;
 
 	/** Create a time/pixel converter for the given geometry */
-	FTimeToPixel GetTimeToPixel(float TrackAreaWidth) const;
+	UE_API FTimeToPixel GetTimeToPixel(float TrackAreaWidth) const;
 
 	/** Gets the current tick resolution of the editor */
-	virtual FFrameRate GetTickResolution() const;
+	UE_API virtual FFrameRate GetTickResolution() const;
 	/** Gets the current view range of this track area */
-	virtual TRange<double> GetViewRange() const;
+	UE_API virtual TRange<double> GetViewRange() const;
 
-	virtual void OnConstruct() override;
+	UE_API virtual void OnConstruct() override;
 
 	/** Get the current active hotspot */
-	TSharedPtr<ITrackAreaHotspot> GetHotspot() const;
+	UE_API TSharedPtr<ITrackAreaHotspot> GetHotspot() const;
 
 	/** Set the hotspot to something else */
-	void SetHotspot(TSharedPtr<ITrackAreaHotspot> NewHotspot);
-	void AddHotspot(TSharedPtr<ITrackAreaHotspot> NewHotspot);
-	void RemoveHotspot(FViewModelTypeID Type);
-	void ClearHotspots();
+	UE_API void SetHotspot(TSharedPtr<ITrackAreaHotspot> NewHotspot);
+	UE_API void AddHotspot(TSharedPtr<ITrackAreaHotspot> NewHotspot);
+	UE_API void RemoveHotspot(FViewModelTypeID Type);
+	UE_API void ClearHotspots();
 
 	/** Set whether the hotspot is locked and cannot be changed (ie when a menu is open) */
-	void LockHotspot(bool bIsLocked);
+	UE_API void LockHotspot(bool bIsLocked);
 
 	/** Get the callback that is fired when the current hotspot changes */
-	FOnTrackAreaHotspotChanged& GetOnHotspotChangedDelegate();
+	UE_API FOnTrackAreaHotspotChanged& GetOnHotspotChangedDelegate();
 
-	void AddEditTool(TSharedPtr<ISequencerEditTool> InNewTool);
+	UE_API void AddEditTool(TSharedPtr<ISequencerEditTool> InNewTool);
 
 	/** Access the currently active track area edit tool */
 	ISequencerEditTool* GetEditTool() const { return EditTool.IsValid() ? EditTool.Get() : nullptr; }
 
 	/** Check whether it's possible to activate the specified tool */
-	bool CanActivateEditTool(FName Identifier) const;
+	UE_API bool CanActivateEditTool(FName Identifier) const;
 
-	bool AttemptToActivateTool(FName Identifier);
+	UE_API bool AttemptToActivateTool(FName Identifier);
 
-	void LockEditTool();
-	void UnlockEditTool();
+	UE_API void LockEditTool();
+	UE_API void UnlockEditTool();
 
-	void AddViewSpace(const FGuid& ID, TSharedPtr<FTrackAreaViewSpace> ViewSpace);
-	void RemoveViewSpace(const FGuid& ID);
+	UE_API void AddViewSpace(const FGuid& ID, TSharedPtr<FTrackAreaViewSpace> ViewSpace);
+	UE_API void RemoveViewSpace(const FGuid& ID);
 
-	TSharedPtr<FTrackAreaViewSpace> FindViewSpace(const FGuid& InID) const;
+	UE_API TSharedPtr<FTrackAreaViewSpace> FindViewSpace(const FGuid& InID) const;
 
 protected:
 
-	void UpdateViewSpaces();
+	UE_API void UpdateViewSpaces();
 
 protected:
 
@@ -121,3 +123,4 @@ protected:
 } // namespace Sequencer
 } // namespace UE
 
+#undef UE_API
diff --git a/Engine/Source/Editor/SequencerCore/Public/MVVM/ViewModels/TrackAreaViewSpace.h b/Engine/Source/Editor/SequencerCore/Public/MVVM/ViewModels/TrackAreaViewSpace.h
--- a/Engine/Source/Editor/SequencerCore/Public/MVVM/ViewModels/TrackAreaViewSpace.h
+++ b/Engine/Source/Editor/SequencerCore/Public/MVVM/ViewModels/TrackAreaViewSpace.h
@@ -6,23 +6,26 @@
 #include "MVVM/ViewModels/ViewModel.h"
 #include "TimeToPixel.h"
 
+#define UE_API SEQUENCERCORE_API
+
 namespace UE::Sequencer
 {
 
-class SEQUENCERCORE_API FTrackAreaViewSpace
+class FTrackAreaViewSpace
 	: public FViewModel
 	, public INonLinearTimeTransform
 {
 public:
 
-	UE_SEQUENCER_DECLARE_CASTABLE(FTrackAreaViewSpace, FViewModel);
+	UE_SEQUENCER_DECLARE_CASTABLE_API(UE_API, FTrackAreaViewSpace, FViewModel);
 
-	FTrackAreaViewSpace();
-	virtual ~FTrackAreaViewSpace();
+	UE_API FTrackAreaViewSpace();
+	UE_API virtual ~FTrackAreaViewSpace();
 
-	virtual double SourceToView(double Seconds) const override;
-	virtual double ViewToSource(double Source) const override;
+	UE_API virtual double SourceToView(double Seconds) const override;
+	UE_API virtual double ViewToSource(double Source) const override;
 };
 
 } // namespace UE::Sequencer
 
+#undef UE_API
diff --git a/Engine/Source/Editor/SequencerCore/Public/MVVM/ViewModels/ViewModel.h b/Engine/Source/Editor/SequencerCore/Public/MVVM/ViewModels/ViewModel.h
--- a/Engine/Source/Editor/SequencerCore/Public/MVVM/ViewModels/ViewModel.h
+++ b/Engine/Source/Editor/SequencerCore/Public/MVVM/ViewModels/ViewModel.h
@@ -15,6 +15,8 @@
 #include "Templates/TypeHash.h"
 #include "Templates/UnrealTemplate.h"
 
+#define UE_API SEQUENCERCORE_API
+
 namespace UE
 {
 namespace Sequencer
@@ -32,7 +34,7 @@ struct FViewModelListIterator;
  * WARNING: most of the casting and hierarchy related methods are not safe to call from a sub-class'
  * constructor because they rely on shared-pointers to be available for `this`.
  */
-class SEQUENCERCORE_API FViewModel
+class FViewModel
 	: public ICastable
 	, public TSharedFromThis<FViewModel>
 	, public FDynamicExtensionContainer
@@ -40,16 +42,16 @@ class SEQUENCERCORE_API FViewModel
 public:
 
 	/** Builds a new data model */
-	FViewModel();
+	UE_API FViewModel();
 
 	FViewModel(const FViewModel&) = delete;
 	FViewModel& operator=(const FViewModel&) = delete;
 
-	virtual ~FViewModel();
+	UE_API virtual ~FViewModel();
 
 public:
 
-	UE_SEQUENCER_DECLARE_CASTABLE(FViewModel, FDynamicExtensionContainer);
+	UE_SEQUENCER_DECLARE_CASTABLE_API(UE_API, FViewModel, FDynamicExtensionContainer);
 
 	/** Adds a dynamic extension to this data model */
 	template<typename T, typename... InArgTypes>
@@ -122,7 +124,7 @@ public:
 	}
 
 	/** Gets the parent data model */
-	FViewModelPtr GetParent() const;
+	UE_API FViewModelPtr GetParent() const;
 
 	/** Gets the parent data model */
 	template<typename T>
@@ -133,7 +135,7 @@ public:
 	}
 
 	/** Returns the root model of the hierarchy, i.e. the parent model that has no parent */
-	FViewModelPtr GetRoot() const;
+	UE_API FViewModelPtr GetRoot() const;
 
 
 	/** Get this model's unique, non-persistent and non-deterministic ID */
@@ -144,37 +146,37 @@ public:
 
 public:
 
-	bool IsConstructed() const;
+	UE_API bool IsConstructed() const;
 
-	TSharedPtr<FSharedViewModelData> GetSharedData() const;
-	void SetSharedData(TSharedPtr<FSharedViewModelData> InSharedData);
+	UE_API TSharedPtr<FSharedViewModelData> GetSharedData() const;
+	UE_API void SetSharedData(TSharedPtr<FSharedViewModelData> InSharedData);
 
 	/** Retrieve the first child list that matches the specified type (if any) */
-	TOptional<FViewModelChildren> FindChildList(EViewModelListType InType);
+	UE_API TOptional<FViewModelChildren> FindChildList(EViewModelListType InType);
 	/** Retrieve the first child list that matches the specified type. Asserts on failure. */
-	FViewModelChildren GetChildList(EViewModelListType InType);
+	UE_API FViewModelChildren GetChildList(EViewModelListType InType);
 
 	/** Returns whether this model has any children in any of its child lists */
-	bool HasChildren() const;
+	UE_API bool HasChildren() const;
 
 	/** Gets the next sibling of this data model in the children list it belongs to */
-	TSharedPtr<FViewModel> GetPreviousSibling() const;
+	UE_API TSharedPtr<FViewModel> GetPreviousSibling() const;
 
 	/** Gets the next sibling of this data model in the children list it belongs to */
-	TSharedPtr<FViewModel> GetNextSibling() const;
+	UE_API TSharedPtr<FViewModel> GetNextSibling() const;
 
 public:
 
 	/** Remove this model from its parent */
-	void RemoveFromParent();
+	UE_API void RemoveFromParent();
 
 	/** Unlinks the head of the children list */
-	void DiscardAllChildren();
+	UE_API void DiscardAllChildren();
 
 public:
 
 	/** Gets an iterator for the data model's default children list */
-	FViewModelListIterator GetChildren(EViewModelListType InFilter = EViewModelListType::Everything) const;
+	UE_API FViewModelListIterator GetChildren(EViewModelListType InFilter = EViewModelListType::Everything) const;
 
 	/** Gets children of a given type from the view-model's children lists, or the specified list */
 	template<typename T>
@@ -186,12 +188,12 @@ public:
 	/**
 	 * Gets all children and their descendants in the hierarchy below this data model.
 	 */
-	FParentFirstChildIterator GetDescendants(bool bIncludeThis = false, EViewModelListType InFilter = EViewModelListType::Everything) const;
+	UE_API FParentFirstChildIterator GetDescendants(bool bIncludeThis = false, EViewModelListType InFilter = EViewModelListType::Everything) const;
 
 	/**
 	 * Gets all children and descendants of a given type in the hierarchy below this data model.
 	 */
-	void GetDescendantsOfType(FViewModelTypeID Type, TArray<TSharedPtr<FViewModel>>& OutChildren, EViewModelListType InFilter = EViewModelListType::Everything) const;
+	UE_API void GetDescendantsOfType(FViewModelTypeID Type, TArray<TSharedPtr<FViewModel>>& OutChildren, EViewModelListType InFilter = EViewModelListType::Everything) const;
 
 	/**
 	 * Gets all children and descendants of a given type in the hierarchy below this data model.
@@ -218,18 +220,18 @@ public:
 	/**
 	 * Finds a parent data model of a given type among the ancestry of this data model.
 	 */
-	FViewModelPtr FindAncestorOfType(FViewModelTypeID Type, bool bIncludeThis = false) const;
+	UE_API FViewModelPtr FindAncestorOfType(FViewModelTypeID Type, bool bIncludeThis = false) const;
 
 	/**
 	 * Finds a parent data model that supports all the given interface types among the ancestry
 	 * of this data model.
 	 */
-	FViewModelPtr FindAncestorOfTypes(TArrayView<const FViewModelTypeID> Types, bool bIncludeThis = false) const;
+	UE_API FViewModelPtr FindAncestorOfTypes(TArrayView<const FViewModelTypeID> Types, bool bIncludeThis = false) const;
 
 	/**
 	 * Iterates over the ancestry of this data model.
 	 */
-	FParentModelIterator GetAncestors(bool bIncludeThis = false) const;
+	UE_API FParentModelIterator GetAncestors(bool bIncludeThis = false) const;
 
 	/**
 	 * Iterates over all ancestors of a given time along the ancestry of this data model.
@@ -243,12 +245,12 @@ public:
 protected:
 
 	/** Chain a new children list to the existing children lists */
-	void RegisterChildList(FViewModelListHead* InChildren);
+	UE_API void RegisterChildList(FViewModelListHead* InChildren);
 
 	/** Retrieve the children for the specified child list head */
-	FViewModelChildren GetChildrenForList(FViewModelListHead* ListHead);
+	UE_API FViewModelChildren GetChildrenForList(FViewModelListHead* ListHead);
 
-	void SetSharedDataFromConstructor(TSharedPtr<FSharedViewModelData> InSharedData);
+	UE_API void SetSharedDataFromConstructor(TSharedPtr<FSharedViewModelData> InSharedData);
 
 private:
 
@@ -258,7 +260,7 @@ private:
 		ChangedParent,
 		ClearedParent,
 	};
-	ESetParentResult SetParentOnly(const TSharedPtr<FViewModel>& NewParent, bool bReportChanges = true);
+	UE_API ESetParentResult SetParentOnly(const TSharedPtr<FViewModel>& NewParent, bool bReportChanges = true);
 
 	/** Called once for this model when it is added to a parent that has valid shared data */
 	virtual void OnConstruct() {}
@@ -304,3 +306,4 @@ private:
 } // namespace Sequencer
 } // namespace UE
 
+#undef UE_API
diff --git a/Engine/Source/Editor/SequencerCore/Public/MVVM/ViewModels/ViewModelHierarchy.h b/Engine/Source/Editor/SequencerCore/Public/MVVM/ViewModels/ViewModelHierarchy.h
--- a/Engine/Source/Editor/SequencerCore/Public/MVVM/ViewModels/ViewModelHierarchy.h
+++ b/Engine/Source/Editor/SequencerCore/Public/MVVM/ViewModels/ViewModelHierarchy.h
@@ -10,6 +10,8 @@
 #include "SequencerCoreFwd.h"
 #include "Templates/SharedPointer.h"
 
+#define UE_API SEQUENCERCORE_API
+
 namespace UE
 {
 namespace Sequencer
@@ -66,15 +68,15 @@ SEQUENCERCORE_API EViewModelListType RegisterCustomModelListType();
  *	which can lead to problems. Use the methods on FViewModel.
  *	This link structure is mostly private anyway.
  */
-struct SEQUENCERCORE_API FViewModelListLink
+struct FViewModelListLink
 {
-	~FViewModelListLink();
+	UE_API ~FViewModelListLink();
 
 	/** Detach only this link from its list */
-	void Unlink();
+	UE_API void Unlink();
 
 	/** Find ths last valid link in the list */
-	TSharedPtr<FViewModel> FindLastLink() const;
+	UE_API TSharedPtr<FViewModel> FindLastLink() const;
 
 private:
 
@@ -83,13 +85,13 @@ private:
 	friend FViewModelChildren;
 
 	/** Throws an assert if a cycle is detected */
-	void DetectLinkListCycle();
+	UE_API void DetectLinkListCycle();
 
 	/** Throws an assert if a cycle is detected */
-	static void DetectLinkListCycle(TSharedPtr<FViewModel> StartAt);
+	static UE_API void DetectLinkListCycle(TSharedPtr<FViewModel> StartAt);
 
 	/** Unlinks the model from its siblings and insert it after the given chain link */
-	static void LinkModelTo(TSharedPtr<FViewModel> Model, TSharedPtr<FViewModelListLink> ToLink);
+	static UE_API void LinkModelTo(TSharedPtr<FViewModel> Model, TSharedPtr<FViewModelListLink> ToLink);
 
 	/** Next data model in the list */
 	TSharedPtr<FViewModel> Next;
@@ -112,20 +114,20 @@ private:
  *	which can lead to problems. Use the methods on FViewModel.
  *	This link structure is mostly private anyway.
  */
-struct SEQUENCERCORE_API FViewModelListHead
+struct FViewModelListHead
 {
 	FViewModelListHead(EViewModelListType InType)
 		: Type(InType)
 	{}
 
 	/** Returns the head model of this list */
-	FViewModelPtr GetHead() const;
+	UE_API FViewModelPtr GetHead() const;
 
 	/** Finds the last model of this list */
-	FViewModelPtr GetTail() const;
+	UE_API FViewModelPtr GetTail() const;
 
 	/** Iterate the children held within this sub list */
-	FViewModelSubListIterator Iterate() const;
+	UE_API FViewModelSubListIterator Iterate() const;
 
 	/** Iterate the children held within this sub list */
 	template<typename T>
@@ -138,7 +140,7 @@ struct SEQUENCERCORE_API FViewModelListHead
 	 * Move all the children in this list to a new list.
 	 * @note: does not change parent relationships - NewHead must be within the same parent
 	 */
-	TSharedPtr<FViewModel> ReliquishList();
+	UE_API TSharedPtr<FViewModel> ReliquishList();
 
 	/** Head of this list, where Next is the first item */
 	FViewModelListLink HeadLink;
@@ -152,12 +154,12 @@ struct SEQUENCERCORE_API FViewModelListHead
  * Scoped object that adds a temporary child list to the specified model of the specified type.
  * Primarily used during regeneration of children to temporarily introduce a 'Recycled' child list
  */
-struct SEQUENCERCORE_API FScopedViewModelListHead
+struct FScopedViewModelListHead
 {
 	/** Constructor that adds the list to the model */
-	FScopedViewModelListHead(const TSharedPtr<FViewModel>& InModel, EViewModelListType InType);
+	UE_API FScopedViewModelListHead(const TSharedPtr<FViewModel>& InModel, EViewModelListType InType);
 	/** Destructor that removes the list to the model */
-	~FScopedViewModelListHead();
+	UE_API ~FScopedViewModelListHead();
 
 	/** Non-copyable and non-moveable since a pointer to our ListHead member is added to our owner */
 	FScopedViewModelListHead(const FScopedViewModelListHead&) = delete;
@@ -169,7 +171,7 @@ struct SEQUENCERCORE_API FScopedViewModelListHead
 	/**
 	 * Retrieve this child list
 	 */
-	FViewModelChildren GetChildren();
+	UE_API FViewModelChildren GetChildren();
 
 private:
 
@@ -179,34 +181,34 @@ private:
 	FViewModelListHead ListHead;
 };
 
-struct SEQUENCERCORE_API FViewModelChildren
+struct FViewModelChildren
 {
 	/**
 	 * Get the collective type of these children
 	 */
-	EViewModelListType GetType() const;
+	UE_API EViewModelListType GetType() const;
 
 	/**
 	 * Get this children's parent
 	 */
-	FViewModelPtr GetParent() const;
+	UE_API FViewModelPtr GetParent() const;
 
 	/**
 	 * Add a child from a different parent to these children
 	 * @note: If the child is already under this parent, this function is a noop
 	 */
-	void AddChild(const TSharedPtr<FViewModel>& Child);
+	UE_API void AddChild(const TSharedPtr<FViewModel>& Child);
 
 	/**
 	 * Add a child from a different parent to these children to the specified previous sibling, or the head if nullptr
 	 * @note: If the child is already the specified model's next sibling, this function is a noop
 	 */
-	void InsertChild(const TSharedPtr<FViewModel>& Child, const TSharedPtr<FViewModel>& PreviousSibling);
+	UE_API void InsertChild(const TSharedPtr<FViewModel>& Child, const TSharedPtr<FViewModel>& PreviousSibling);
 
 	/**
 	 * Move all these children into the specified destination list, changing parent pointers if necessary
 	 */
-	void MoveChildrenTo(const FViewModelChildren& OutDestination);
+	UE_API void MoveChildrenTo(const FViewModelChildren& OutDestination);
 
 	/**
 	 * Invoke a callback on these children, and move them all into the specified destination list, changing parent pointers if necessary
@@ -224,12 +226,12 @@ struct SEQUENCERCORE_API FViewModelChildren
 	/**
 	 * Empty this specific sub-list, resetting all parent ptrs
 	 */
-	void Empty();
+	UE_API void Empty();
 
 public:
 
 	/** Returns an iterator over this list */
-	FViewModelSubListIterator IterateSubList() const;
+	UE_API FViewModelSubListIterator IterateSubList() const;
 
 	/** Returns an iterator over this list which will skip any item not of type T */
 	template<typename T>
@@ -245,10 +247,10 @@ public:
 	}
 
 	/** Returns the head model of this list */
-	FViewModelPtr GetHead() const;
+	UE_API FViewModelPtr GetHead() const;
 
 	/** Finds the last model of this list */
-	FViewModelPtr GetTail() const;
+	UE_API FViewModelPtr GetTail() const;
 
 	template<typename T>
 	TViewModelPtr<T> FindFirstChildOfType() const
@@ -265,7 +267,7 @@ private:
 	friend FViewModel;
 	friend FScopedViewModelListHead;
 
-	FViewModelChildren(const TSharedPtr<FViewModel>& InOwner, FViewModelListHead* InListHead);
+	UE_API FViewModelChildren(const TSharedPtr<FViewModel>& InOwner, FViewModelListHead* InListHead);
 
 	TSharedPtr<FViewModel> Owner;
 	FViewModelListHead* ListHead;
@@ -287,7 +289,7 @@ private:
  * be a maximum of one event triggered for each changed node. This is achieved by tracking
  * which models have already been handled during the event processing.
  */
-class SEQUENCERCORE_API FViewModelHierarchyOperation
+class FViewModelHierarchyOperation
 {
 public:
 
@@ -299,17 +301,17 @@ public:
 		Construct();
 	}
 
-	explicit FViewModelHierarchyOperation(const TSharedPtr<FSharedViewModelData>& InSharedData);
-	~FViewModelHierarchyOperation();
+	UE_API explicit FViewModelHierarchyOperation(const TSharedPtr<FSharedViewModelData>& InSharedData);
+	UE_API ~FViewModelHierarchyOperation();
 
 	/**
 	 * Called before a change is made to the specified model's parent or previous sibling
 	 */
-	void PreHierarchicalChange(const TSharedPtr<FViewModel>& InChangedNode);
+	UE_API void PreHierarchicalChange(const TSharedPtr<FViewModel>& InChangedNode);
 
 private:
 
-	void Construct();
+	UE_API void Construct();
 
 	struct FCachedHierarchicalPosition
 	{
@@ -334,3 +336,4 @@ private:
 } // namespace Sequencer
 } // namespace UE
 
+#undef UE_API
diff --git a/Engine/Source/Editor/SequencerCore/Public/MVVM/ViewModels/VirtualTrackArea.h b/Engine/Source/Editor/SequencerCore/Public/MVVM/ViewModels/VirtualTrackArea.h
--- a/Engine/Source/Editor/SequencerCore/Public/MVVM/ViewModels/VirtualTrackArea.h
+++ b/Engine/Source/Editor/SequencerCore/Public/MVVM/ViewModels/VirtualTrackArea.h
@@ -9,6 +9,8 @@
 #include "Templates/SharedPointer.h"
 #include "TimeToPixel.h"
 
+#define UE_API SEQUENCERCORE_API
+
 namespace UE::Sequencer { class FViewModel; }
 
 namespace UE
@@ -20,32 +22,32 @@ class FTrackAreaViewModel;
 class SOutlinerView;
 
 /** Structure used for handling the virtual space of the track area */
-class SEQUENCERCORE_API FVirtualTrackArea : public FTimeToPixel
+class FVirtualTrackArea : public FTimeToPixel
 {
 public:
 
 	/** Construction responsibility is delegated to SSequencer. See SSequencer::GetVirtualTrackArea */
-	FVirtualTrackArea(const FTrackAreaViewModel& InTrackArea, SOutlinerView& InTreeView, const FGeometry& InTrackAreaGeometry);
+	UE_API FVirtualTrackArea(const FTrackAreaViewModel& InTrackArea, SOutlinerView& InTreeView, const FGeometry& InTrackAreaGeometry);
 
 	/** Convert the specified pixel position into a virtual vertical offset from the absolute top of the tree */
-	float PixelToVerticalOffset(float InPixel) const;
+	UE_API float PixelToVerticalOffset(float InPixel) const;
 
 	/** Convert the specified absolute vertical position into a physical vertical offset in the track area. */
 	/** @note: Use with caution - not reliable where the specified offset is not on screen */
-	float VerticalOffsetToPixel(float InOffset) const;
+	UE_API float VerticalOffsetToPixel(float InOffset) const;
 
 	/** Convert the specified physical point into a virtual point the absolute top of the tree */
-	FVector2D PhysicalToVirtual(FVector2D InPosition) const;
+	UE_API FVector2D PhysicalToVirtual(FVector2D InPosition) const;
 
 	/** Convert the specified absolute virtual point into a physical point in the track area. */
 	/** @note: Use with caution - not reliable where the specified point is not on screen */
-	FVector2D VirtualToPhysical(FVector2D InPosition) const;
+	UE_API FVector2D VirtualToPhysical(FVector2D InPosition) const;
 
 	/** Get the physical size of the track area */
-	FVector2D GetPhysicalSize() const;
+	UE_API FVector2D GetPhysicalSize() const;
 
 	/** Hit test at the specified physical position for a sequencer node */
-	TSharedPtr<FViewModel> HitTestNode(float InPhysicalPosition) const;
+	UE_API TSharedPtr<FViewModel> HitTestNode(float InPhysicalPosition) const;
 
 	/** Cached track area geometry */
 	FGeometry CachedTrackAreaGeometry() const { return TrackAreaGeometry; }
@@ -62,3 +64,4 @@ private:
 } // namespace Sequencer
 } // namespace UE
 
+#undef UE_API
diff --git a/Engine/Source/Editor/SequencerCore/Public/MVVM/Views/ITrackAreaHotspot.h b/Engine/Source/Editor/SequencerCore/Public/MVVM/Views/ITrackAreaHotspot.h
--- a/Engine/Source/Editor/SequencerCore/Public/MVVM/Views/ITrackAreaHotspot.h
+++ b/Engine/Source/Editor/SequencerCore/Public/MVVM/Views/ITrackAreaHotspot.h
@@ -9,6 +9,8 @@
 
 #include "MVVM/ICastable.h"
 
+#define UE_API SEQUENCERCORE_API
+
 class FExtender;
 class FMenuBuilder;
 struct FGeometry;
@@ -23,10 +25,10 @@ class FTrackAreaViewModel;
 class ISequencerEditToolDragOperation;
 
 /** A sequencer hotspot is used to identify specific areas on the sequencer track area */ 
-struct SEQUENCERCORE_API ITrackAreaHotspot 
+struct ITrackAreaHotspot 
 	: public ICastable
 {
-	UE_SEQUENCER_DECLARE_CASTABLE(ITrackAreaHotspot);
+	UE_SEQUENCER_DECLARE_CASTABLE_API(UE_API, ITrackAreaHotspot);
 
 	virtual ~ITrackAreaHotspot() { }
 	virtual void UpdateOnHover(FTrackAreaViewModel& InTrackArea) const = 0;
@@ -58,3 +60,4 @@ FORCEINLINE TSharedPtr<T> HotspotCast(const TSharedPtr<ITrackAreaHotspot>& InHot
 } // namespace Sequencer
 } // namespace UE
 
+#undef UE_API
diff --git a/Engine/Source/Editor/SequencerCore/Public/MVVM/Views/KeyRenderer.h b/Engine/Source/Editor/SequencerCore/Public/MVVM/Views/KeyRenderer.h
--- a/Engine/Source/Editor/SequencerCore/Public/MVVM/Views/KeyRenderer.h
+++ b/Engine/Source/Editor/SequencerCore/Public/MVVM/Views/KeyRenderer.h
@@ -15,6 +15,8 @@
 #include "TimeToPixel.h"
 #include "MVVM/ViewModelPtr.h"
 
+#define UE_API SEQUENCERCORE_API
+
 struct FSlateBrush;
 class FWidgetStyle;
 class FSlateWindowElementList;
@@ -136,7 +138,7 @@ struct FKeyBatchParameters
 /**
  * Utility class for efficiently drawing large numbers of keys on a track lane.
  */
-struct SEQUENCERCORE_API FKeyRenderer
+struct FKeyRenderer
 {
 	/** Flag enum signifying states for a particular key or group of keys */
 	enum class EKeyRenderingFlags
@@ -180,37 +182,37 @@ struct SEQUENCERCORE_API FKeyRenderer
 	/**
 	 * Initialize this renderer with a generic model
 	 */
-	void Initialize(const FViewModelPtr& InViewModel);
+	UE_API void Initialize(const FViewModelPtr& InViewModel);
 
 	/**
 	 * Update any invalidated data required for drawing keys
 	 */
-	void Update(const FKeyBatchParameters& Params, const FGeometry& AllottedGeometry) const;
+	UE_API void Update(const FKeyBatchParameters& Params, const FGeometry& AllottedGeometry) const;
 
 	/**
 	 * Draw this batch's curve
 	 */
-	int32 DrawCurve(const FKeyBatchParameters& Params, const FGeometry& AllottedGeometry, const FSlateRect& MyCullingRect, const FKeyRendererPaintArgs& Args, int32 LayerId) const;
+	UE_API int32 DrawCurve(const FKeyBatchParameters& Params, const FGeometry& AllottedGeometry, const FSlateRect& MyCullingRect, const FKeyRendererPaintArgs& Args, int32 LayerId) const;
 
 	/**
 	 * Draw this batch
 	 */
-	int32 Draw(const FKeyBatchParameters& Params, const FGeometry& AllottedGeometry, const FSlateRect& MyCullingRect, const FKeyRendererPaintArgs& Args, int32 LayerId) const;
+	UE_API int32 Draw(const FKeyBatchParameters& Params, const FGeometry& AllottedGeometry, const FSlateRect& MyCullingRect, const FKeyRendererPaintArgs& Args, int32 LayerId) const;
 
 	/**
 	 * Retrieve keys at the specified position
 	 */
-	void HitTestKeys(const FFrameTime& Time, TArray<FKeysForModel>& OutAllKeys) const;
+	UE_API void HitTestKeys(const FFrameTime& Time, TArray<FKeysForModel>& OutAllKeys) const;
 
 	/**
 	 * Retrieve keys at the specified position
 	 */
-	bool HitTestKeyBar(const FFrameTime& Time, FKeyBar& OutKeyBar) const;
+	UE_API bool HitTestKeyBar(const FFrameTime& Time, FKeyBar& OutKeyBar) const;
 
 	/**
 	 * Returns true if this renderer has a curve, false otherwise
 	 */
-	bool HasCurve() const;
+	UE_API bool HasCurve() const;
 
 private:
 
@@ -310,19 +312,19 @@ private:
 	/**
 	 * Cache all the key extensions from our model
 	 */
-	void CacheKeyExtensions(const FKeyBatchParameters& Params) const;
+	UE_API void CacheKeyExtensions(const FKeyBatchParameters& Params) const;
 
 	/**
 	 * Attempt to update data that is not dependent upon the current view
 	 * @return EViewDependentCacheFlags::DataChanged if view-indipendent data was updated, EViewDependentCacheFlags::None otherwise
 	 */
-	EViewDependentCacheFlags UpdateViewIndependentData() const;
+	UE_API EViewDependentCacheFlags UpdateViewIndependentData() const;
 
 
 	/**
 	 * Ensure that view-dependent data (such as which keys need drawing and how) is up to date
 	 */
-	void UpdateViewDependentData(const FKeyBatchParameters& Params, const FGeometry& AllottedGeometry) const;
+	UE_API void UpdateViewDependentData(const FKeyBatchParameters& Params, const FGeometry& AllottedGeometry) const;
 
 private:
 
@@ -350,3 +352,4 @@ ENUM_CLASS_FLAGS(FKeyRenderer::EKeyRenderingFlags);
 
 } // namespace UE::Sequencer
 
+#undef UE_API
diff --git a/Engine/Source/Editor/SequencerCore/Public/MVVM/Views/OutlinerColumns/SColumnToggleWidget.h b/Engine/Source/Editor/SequencerCore/Public/MVVM/Views/OutlinerColumns/SColumnToggleWidget.h
--- a/Engine/Source/Editor/SequencerCore/Public/MVVM/Views/OutlinerColumns/SColumnToggleWidget.h
+++ b/Engine/Source/Editor/SequencerCore/Public/MVVM/Views/OutlinerColumns/SColumnToggleWidget.h
@@ -7,6 +7,8 @@
 #include "ScopedTransaction.h"
 #include "Templates/UniquePtr.h"
 
+#define UE_API SEQUENCERCORE_API
+
 namespace UE::Sequencer
 {
 
@@ -22,7 +24,7 @@ class FSequencerEditorViewModel;
  * Column Toggle widgets have 4 States: Active, Inactive, Active Child, and Implicitly Active.
  * Column Toggle widgets will change its brush based on its state.
  */
-class SEQUENCERCORE_API SColumnToggleWidget
+class SColumnToggleWidget
 	: public SImage
 {
 public:
@@ -30,7 +32,7 @@ public:
 	SLATE_BEGIN_ARGS(SColumnToggleWidget) {}
 	SLATE_END_ARGS()
 
-	void Construct(
+	UE_API void Construct(
 		const FArguments& InArgs,
 		const TWeakPtr<IOutlinerColumn> InOutlinerColumn,
 		const FCreateOutlinerColumnParams& InParams);
@@ -41,7 +43,7 @@ public:
 	virtual void OnToggleOperationComplete() {}
 
 	/** Get the color and opacity of the column toggle widget. */
-	virtual FSlateColor GetImageColorAndOpacity() const;
+	UE_API virtual FSlateColor GetImageColorAndOpacity() const;
 
 protected:
 
@@ -72,44 +74,44 @@ protected:
 	
 
 	/** Get the image this widget displays. */
-	virtual const FSlateBrush* GetBrush() const;
+	UE_API virtual const FSlateBrush* GetBrush() const;
 
 protected:
 
 	/** Start a new drag/drop operation for this widget. */
-	virtual FReply OnDragDetected(const FGeometry& MyGeometry, const FPointerEvent& MouseEvent) override;
+	UE_API virtual FReply OnDragDetected(const FGeometry& MyGeometry, const FPointerEvent& MouseEvent) override;
 
 	/** If a matching column drag drop operation has entered this widget, set its item to the specified active state. */
-	virtual void OnDragEnter(const FGeometry& MyGeometry, const FDragDropEvent& DragDropEvent) override;
+	UE_API virtual void OnDragEnter(const FGeometry& MyGeometry, const FDragDropEvent& DragDropEvent) override;
 
 	/** Handles left click-type inputs and potentially begins drag and drop operation. */
-	FReply HandleClick();
+	UE_API FReply HandleClick();
 
 	/** Called when the mouse button is double clicked on this widget. */
-	virtual FReply OnMouseButtonDoubleClick(const FGeometry& InMyGeometry, const FPointerEvent& InMouseEvent) override;
+	UE_API virtual FReply OnMouseButtonDoubleClick(const FGeometry& InMyGeometry, const FPointerEvent& InMouseEvent) override;
 
 	/** Called when the mouse button is pressed down on this widget. */
-	virtual FReply OnMouseButtonDown(const FGeometry& MyGeometry, const FPointerEvent& MouseEvent) override;
+	UE_API virtual FReply OnMouseButtonDown(const FGeometry& MyGeometry, const FPointerEvent& MouseEvent) override;
 
 	/** Called when this widget had captured the mouse, but that capture has been revoked for some reason. */
-	virtual void OnMouseCaptureLost(const FCaptureLostEvent& CaptureLostEvent) override;
+	UE_API virtual void OnMouseCaptureLost(const FCaptureLostEvent& CaptureLostEvent) override;
 
 	/** Process a mouse up message. */
-	virtual FReply OnMouseButtonUp(const FGeometry& MyGeometry, const FPointerEvent& MouseEvent) override;
+	UE_API virtual FReply OnMouseButtonUp(const FGeometry& MyGeometry, const FPointerEvent& MouseEvent) override;
 
 	/** Tick this widget to update its cached state */
-	virtual void Tick(const FGeometry& AllottedGeometry, const double InCurrentTime, const float InDeltaTime) override;
+	UE_API virtual void Tick(const FGeometry& AllottedGeometry, const double InCurrentTime, const float InDeltaTime) override;
 
 	/** Called when mouse enters the column widget. */
-	void OnMouseEnter(const FGeometry& MyGeometry, const FPointerEvent& MouseEvent) override;
+	UE_API void OnMouseEnter(const FGeometry& MyGeometry, const FPointerEvent& MouseEvent) override;
 
 	/** Called when mouse leaves the column widget. */
-	void OnMouseLeave(const FPointerEvent& MouseEvent) override;
+	UE_API void OnMouseLeave(const FPointerEvent& MouseEvent) override;
 
 protected:
 
 	/** Refreshes the Sequencer Tree. Can be used after toggle operation is complete if operation affects other items. */
-	void RefreshSequencerTree();
+	UE_API void RefreshSequencerTree();
 
 protected:
 	/** Editor view-model for this view. */
@@ -151,4 +153,6 @@ private:
 	uint8 bIsImplicitlyActive : 1;
 };
 
-} // namespace UE::Sequencer
\ No newline at end of file
+} // namespace UE::Sequencer
+
+#undef UE_API
diff --git a/Engine/Source/Editor/SequencerCore/Public/MVVM/Views/OutlinerColumns/SOutlinerColumnButton.h b/Engine/Source/Editor/SequencerCore/Public/MVVM/Views/OutlinerColumns/SOutlinerColumnButton.h
--- a/Engine/Source/Editor/SequencerCore/Public/MVVM/Views/OutlinerColumns/SOutlinerColumnButton.h
+++ b/Engine/Source/Editor/SequencerCore/Public/MVVM/Views/OutlinerColumns/SOutlinerColumnButton.h
@@ -5,6 +5,8 @@
 #include "Widgets/SCompoundWidget.h"
 #include "Framework/SlateDelegates.h"
 
+#define UE_API SEQUENCERCORE_API
+
 class SMenuAnchor;
 
 namespace UE::Sequencer
@@ -12,7 +14,7 @@ namespace UE::Sequencer
 
 class IOutlinerExtension;
 
-class SEQUENCERCORE_API SOutlinerColumnButton : public SCompoundWidget
+class SOutlinerColumnButton : public SCompoundWidget
 {
 public:
 
@@ -30,13 +32,13 @@ public:
 
 	SLATE_END_ARGS()
 
-	void Construct(const FArguments& InArgs);
+	UE_API void Construct(const FArguments& InArgs);
 
 private:
 
-	FSlateColor GetColorAndOpacity() const;
+	UE_API FSlateColor GetColorAndOpacity() const;
 
-	FReply ToggleMenu();
+	UE_API FReply ToggleMenu();
 
 private:
 
@@ -44,4 +46,6 @@ private:
 	TAttribute<bool> IsRowHovered;
 };
 
-} // namespace UE::Sequencer
\ No newline at end of file
+} // namespace UE::Sequencer
+
+#undef UE_API
diff --git a/Engine/Source/Editor/SequencerCore/Public/MVVM/Views/SCompoundTrackLaneView.h b/Engine/Source/Editor/SequencerCore/Public/MVVM/Views/SCompoundTrackLaneView.h
--- a/Engine/Source/Editor/SequencerCore/Public/MVVM/Views/SCompoundTrackLaneView.h
+++ b/Engine/Source/Editor/SequencerCore/Public/MVVM/Views/SCompoundTrackLaneView.h
@@ -8,6 +8,8 @@
 #include "Widgets/SPanel.h"
 #include "Layout/Children.h"
 
+#define UE_API SEQUENCERCORE_API
+
 struct FGeometry;
 
 class FPaintArgs;
@@ -24,7 +26,7 @@ class ITrackLaneWidgetSpace;
 /**
  * 
  */
-class SEQUENCERCORE_API SCompoundTrackLaneView
+class SCompoundTrackLaneView
 	: public SPanel
 {
 public:
@@ -32,8 +34,8 @@ public:
 	SLATE_BEGIN_ARGS(SCompoundTrackLaneView){}
 	SLATE_END_ARGS()
 
-	SCompoundTrackLaneView();
-	~SCompoundTrackLaneView();
+	UE_API SCompoundTrackLaneView();
+	UE_API ~SCompoundTrackLaneView();
 
 	UE_DEPRECATED(5.6, "Please use the TSharedPtr<ITrackLaneWidgetSpace> overload")
 	void Construct(const FArguments& InArgs)
@@ -41,17 +43,17 @@ public:
 		Construct(InArgs, nullptr);
 	}
 
-	void Construct(const FArguments& InArgs, TSharedPtr<ITrackLaneWidgetSpace> InTrackLaneWidgetSpace);
+	UE_API void Construct(const FArguments& InArgs, TSharedPtr<ITrackLaneWidgetSpace> InTrackLaneWidgetSpace);
 
-	void AddWeakWidget(TSharedPtr<ITrackLaneWidget> InWidget, TWeakPtr<STrackLane> InOwningLane);
-	void AddStrongWidget(TSharedPtr<ITrackLaneWidget> InWidget, TWeakPtr<STrackLane> InOwningLane);
+	UE_API void AddWeakWidget(TSharedPtr<ITrackLaneWidget> InWidget, TWeakPtr<STrackLane> InOwningLane);
+	UE_API void AddStrongWidget(TSharedPtr<ITrackLaneWidget> InWidget, TWeakPtr<STrackLane> InOwningLane);
 
 	/*~ SPanel Interface */
-	void OnArrangeChildren( const FGeometry& AllottedGeometry, FArrangedChildren& ArrangedChildren ) const override;
-	FVector2D ComputeDesiredSize(float) const override;
-	FChildren* GetChildren() override;
+	UE_API void OnArrangeChildren( const FGeometry& AllottedGeometry, FArrangedChildren& ArrangedChildren ) const override;
+	UE_API FVector2D ComputeDesiredSize(float) const override;
+	UE_API FChildren* GetChildren() override;
 
-	virtual int32 OnPaint(const FPaintArgs& Args, const FGeometry& AllottedGeometry, const FSlateRect& MyCullingRect, FSlateWindowElementList& OutDrawElements, int32 LayerId, const FWidgetStyle& InWidgetStyle, bool bParentEnabled) const override;
+	UE_API virtual int32 OnPaint(const FPaintArgs& Args, const FGeometry& AllottedGeometry, const FSlateRect& MyCullingRect, FSlateWindowElementList& OutDrawElements, int32 LayerId, const FWidgetStyle& InWidgetStyle, bool bParentEnabled) const override;
 
 private:
 
@@ -79,3 +81,4 @@ private:
 
 } // namespace UE::Sequencer
 
+#undef UE_API
diff --git a/Engine/Source/Editor/SequencerCore/Public/MVVM/Views/SKeyNavigationButtons.h b/Engine/Source/Editor/SequencerCore/Public/MVVM/Views/SKeyNavigationButtons.h
--- a/Engine/Source/Editor/SequencerCore/Public/MVVM/Views/SKeyNavigationButtons.h
+++ b/Engine/Source/Editor/SequencerCore/Public/MVVM/Views/SKeyNavigationButtons.h
@@ -16,6 +16,8 @@
 #include "Widgets/SCompoundWidget.h"
 #include "Misc/EnumClassFlags.h"
 
+#define UE_API SEQUENCERCORE_API
+
 struct FFrameNumber;
 
 namespace UE
@@ -39,7 +41,7 @@ ENUM_CLASS_FLAGS(EKeyNavigationButtons)
 /**
  * A widget for navigating between keys on a sequencer track
  */
-class SEQUENCERCORE_API SKeyNavigationButtons
+class SKeyNavigationButtons
 	: public SCompoundWidget
 {
 public:
@@ -64,11 +66,11 @@ public:
 		SLATE_EVENT(FOnGetNavigatableTimes, GetNavigatableTimes)
 	SLATE_END_ARGS()
 
-	void Construct(const FArguments& InArgs, const TSharedPtr<FViewModel>& InModel);
+	UE_API void Construct(const FArguments& InArgs, const TSharedPtr<FViewModel>& InModel);
 
-	FReply OnPreviousKeyClicked();
-	FReply OnNextKeyClicked();
-	FReply OnAddKeyClicked();
+	UE_API FReply OnPreviousKeyClicked();
+	UE_API FReply OnNextKeyClicked();
+	UE_API FReply OnAddKeyClicked();
 
 protected:
 
@@ -83,3 +85,4 @@ protected:
 } // namespace Sequencer
 } // namespace UE
 
+#undef UE_API
diff --git a/Engine/Source/Editor/SequencerCore/Public/MVVM/Views/SOutlinerItemViewBase.h b/Engine/Source/Editor/SequencerCore/Public/MVVM/Views/SOutlinerItemViewBase.h
--- a/Engine/Source/Editor/SequencerCore/Public/MVVM/Views/SOutlinerItemViewBase.h
+++ b/Engine/Source/Editor/SequencerCore/Public/MVVM/Views/SOutlinerItemViewBase.h
@@ -15,6 +15,8 @@
 #include "Widgets/SCompoundWidget.h"
 #include "Widgets/SWidget.h"
 
+#define UE_API SEQUENCERCORE_API
+
 class ISequencer;
 struct FGeometry;
 struct FOptionalSize;
@@ -42,7 +44,7 @@ enum class EOutlinerItemViewBaseStyle
 /**
  * A widget for displaying a sequencer tree node in the animation outliner.
  */
-class SEQUENCERCORE_API SOutlinerItemViewBase
+class SOutlinerItemViewBase
 	: public SCompoundWidget
 {
 public:
@@ -55,7 +57,7 @@ public:
 		SLATE_NAMED_SLOT(FArguments, RightGutterContent)
 	SLATE_END_ARGS()
 
-	void Construct( 
+	UE_API void Construct( 
 			const FArguments& InArgs, 
 			TWeakViewModelPtr<IOutlinerExtension> InWeakExtension,
 			TWeakPtr<FEditorViewModel> InWeakEditor,
@@ -66,85 +68,85 @@ protected:
 	/**
 	 * @return The color used to draw the display name.
 	 */
-	virtual FText GetLabel() const;
+	UE_API virtual FText GetLabel() const;
 
 	/**
 	 * @return The color used to draw the display name.
 	 */
-	virtual FSlateColor GetLabelColor() const;
+	UE_API virtual FSlateColor GetLabelColor() const;
 
 	/**
 	*@return The font used to draw the display name.
 	*/
-	virtual FSlateFontInfo GetLabelFont() const;
+	UE_API virtual FSlateFontInfo GetLabelFont() const;
 
 	/**
 	 * @return The text displayed for the tool tip for the diplay name label. 
 	 */
-	virtual FText GetLabelToolTipText() const;
+	UE_API virtual FText GetLabelToolTipText() const;
 
 	/**
 	 * Returns whether a requested rename is valid.
 	 */
-	bool IsRenameValid(const FText& NewName, FText& OutErrorMessage) const;
+	UE_API bool IsRenameValid(const FText& NewName, FText& OutErrorMessage) const;
 
 	/**
 	 * Called when a new name has been committed in the text box.
 	 */
-	void OnNodeLabelTextCommitted(const FText& NewLabel, ETextCommit::Type CommitType);
+	UE_API void OnNodeLabelTextCommitted(const FText& NewLabel, ETextCommit::Type CommitType);
 
 	/**
 	 * @return The icon to use on the view
 	 */
-	virtual const FSlateBrush* GetIconBrush() const;
+	UE_API virtual const FSlateBrush* GetIconBrush() const;
 
 	/**
 	 * @return A color tint to apply to the icon
 	 */
-	virtual FSlateColor GetIconTint() const;
+	UE_API virtual FSlateColor GetIconTint() const;
 
 	/**
 	 * @return The icon overlay (eg a spawnable badge) to use on the view
 	 */
-	virtual const FSlateBrush* GetIconOverlayBrush() const;
+	UE_API virtual const FSlateBrush* GetIconOverlayBrush() const;
 
 	/**
 	 * @return Tooltip text to apply to the icon
 	 */
-	virtual FText GetIconToolTipText() const;
+	UE_API virtual FText GetIconToolTipText() const;
 
 	/**
 	 * Callback for checking whether the node label can be edited.
 	 */
-	virtual bool IsNodeLabelReadOnly() const;
+	UE_API virtual bool IsNodeLabelReadOnly() const;
 
 	/**
 	 * Whether this node should be drawn dimmed or not
 	 */
-	virtual bool IsDimmed() const;
+	UE_API virtual bool IsDimmed() const;
 
-	FOptionalSize GetHeight() const;
+	UE_API FOptionalSize GetHeight() const;
 
 private:
 
 	// SWidget interface
 
-	FSlateColor GetForegroundBasedOnSelection() const;
+	UE_API FSlateColor GetForegroundBasedOnSelection() const;
 
 	/**
 	 * @return The border image to show in the tree node.
 	 */
-	const FSlateBrush* GetNodeBorderImage() const;
+	UE_API const FSlateBrush* GetNodeBorderImage() const;
 
 	/**
 	* @return The tint to apply to the border image for the inner portion of the node.
 	*/
-	FSlateColor GetNodeInnerBackgroundTint() const;
+	UE_API FSlateColor GetNodeInnerBackgroundTint() const;
 
 	/**
 	 * @return The expander visibility of this node.
 	 */
-	EVisibility GetExpanderVisibility() const;
+	UE_API EVisibility GetExpanderVisibility() const;
 
 protected:
 
@@ -178,3 +180,4 @@ protected:
 } // namespace Sequencer
 } // namespace UE
 
+#undef UE_API
diff --git a/Engine/Source/Editor/SequencerCore/Public/MVVM/Views/SOutlinerView.h b/Engine/Source/Editor/SequencerCore/Public/MVVM/Views/SOutlinerView.h
--- a/Engine/Source/Editor/SequencerCore/Public/MVVM/Views/SOutlinerView.h
+++ b/Engine/Source/Editor/SequencerCore/Public/MVVM/Views/SOutlinerView.h
@@ -28,6 +28,8 @@
 #include "Widgets/Views/STableViewBase.h"
 #include "Widgets/Views/STreeView.h"
 
+#define UE_API SEQUENCERCORE_API
+
 class FDragDropEvent;
 class FPaintArgs;
 class FSlateRect;
@@ -65,7 +67,7 @@ enum class ETreeRecursion
 
 
 /** The tree view used in the sequencer */
-class SEQUENCERCORE_API SOutlinerView
+class SOutlinerView
 	: public STreeView<TWeakViewModelPtr<IOutlinerExtension>>
 {
 public:
@@ -78,109 +80,109 @@ public:
 
 	SLATE_END_ARGS()
 
-	static const FName TrackNameColumn;
+	static UE_API const FName TrackNameColumn;
 
-	SOutlinerView();
-	~SOutlinerView();
+	UE_API SOutlinerView();
+	UE_API ~SOutlinerView();
 
 	/** Construct this widget */
-	void Construct(const FArguments& InArgs, TWeakPtr<FOutlinerViewModel> InWeakOutliner, const TSharedRef<STrackAreaView>& InTrackArea);
-	virtual void Tick( const FGeometry& AllottedGeometry, const double InCurrentTime, const float InDeltaTime );
-	virtual int32 OnPaint(const FPaintArgs& Args, const FGeometry& AllottedGeometry, const FSlateRect& MyCullingRect, FSlateWindowElementList& OutDrawElements, int32 LayerId, const FWidgetStyle& InWidgetStyle, bool bParentEnabled) const override;
+	UE_API void Construct(const FArguments& InArgs, TWeakPtr<FOutlinerViewModel> InWeakOutliner, const TSharedRef<STrackAreaView>& InTrackArea);
+	UE_API virtual void Tick( const FGeometry& AllottedGeometry, const double InCurrentTime, const float InDeltaTime );
+	UE_API virtual int32 OnPaint(const FPaintArgs& Args, const FGeometry& AllottedGeometry, const FSlateRect& MyCullingRect, FSlateWindowElementList& OutDrawElements, int32 LayerId, const FWidgetStyle& InWidgetStyle, bool bParentEnabled) const override;
 
-	TSharedPtr<FOutlinerViewModel> GetOutlinerModel() const;
+	UE_API TSharedPtr<FOutlinerViewModel> GetOutlinerModel() const;
 
 	/** @return the number of root nodes this tree contains */
 	int32 GetNumRootNodes() const { return RootNodes.Num(); }
 
-	float GetVirtualTop() const;
+	UE_API float GetVirtualTop() const;
 
-	void GetVisibleItems(TArray<TViewModelPtr<IOutlinerExtension>>& OutItems) const;
+	UE_API void GetVisibleItems(TArray<TViewModelPtr<IOutlinerExtension>>& OutItems) const;
 
-	void ForceSetSelectedItems(const TSet<TWeakViewModelPtr<IOutlinerExtension>>& InItems);
+	UE_API void ForceSetSelectedItems(const TSet<TWeakViewModelPtr<IOutlinerExtension>>& InItems);
 
 public:
 
 	/** Get the tree item model at the specified physical vertical position */
-	TViewModelPtr<IOutlinerExtension> HitTestNode(float InPhysical) const;
+	UE_API TViewModelPtr<IOutlinerExtension> HitTestNode(float InPhysical) const;
 
 	/** Convert the specified physical vertical position into an absolute virtual position, ignoring expanded states */
-	float PhysicalToVirtual(float InPhysical) const;
+	UE_API float PhysicalToVirtual(float InPhysical) const;
 
 	/** Convert the specified absolute virtual position into a physical position in the tree.
 	 * @note: Will not work reliably for virtual positions that are outside of the physical space
 	 */
-	float VirtualToPhysical(float InVirtual) const;
+	UE_API float VirtualToPhysical(float InVirtual) const;
 
-	void ReportChildRowGeometry(const TViewModelPtr<IOutlinerExtension>& InNode, const FGeometry& InGeometry);
+	UE_API void ReportChildRowGeometry(const TViewModelPtr<IOutlinerExtension>& InNode, const FGeometry& InGeometry);
 
 public:
 
 	/** Refresh this tree as a result of the underlying tree data changing */
-	void Refresh();
+	UE_API void Refresh();
 
 	/** Expand or collapse nodes */
-	void ToggleExpandCollapseNodes(ETreeRecursion Recursion = ETreeRecursion::Recursive, bool bExpandAll = false, bool bCollapseAll = false);
+	UE_API void ToggleExpandCollapseNodes(ETreeRecursion Recursion = ETreeRecursion::Recursive, bool bExpandAll = false, bool bCollapseAll = false);
 
 	/** Scroll this tree view by the specified number of slate units */
-	void ScrollByDelta(float DeltaInSlateUnits);
+	UE_API void ScrollByDelta(float DeltaInSlateUnits);
 
-	bool IsColumnVisible(const FName& InName) const;
+	UE_API bool IsColumnVisible(const FName& InName) const;
 
 	/** Set the item's expansion state, including all of its children */
-	void ExpandCollapseNode(TViewModelPtr<IOutlinerExtension> InDataModel, bool bExpansionState, ETreeRecursion Recursion);
+	UE_API void ExpandCollapseNode(TViewModelPtr<IOutlinerExtension> InDataModel, bool bExpansionState, ETreeRecursion Recursion);
 
 protected:
 
 	/** Generate a row for a particular node */
-	virtual TSharedRef<ITableRow> OnGenerateRow(TWeakViewModelPtr<IOutlinerExtension> InDisplayNode, const TSharedRef<STableViewBase>& OwnerTable);
+	UE_API virtual TSharedRef<ITableRow> OnGenerateRow(TWeakViewModelPtr<IOutlinerExtension> InDisplayNode, const TSharedRef<STableViewBase>& OwnerTable);
 
-	void CreateTrackLanesForRow(TSharedRef<SOutlinerViewRow> InRow, TViewModelPtr<IOutlinerExtension> InDataModel);
-	TSharedPtr<STrackLane> FindOrCreateParentLane(TViewModelPtr<IOutlinerExtension> InDataModel);
+	UE_API void CreateTrackLanesForRow(TSharedRef<SOutlinerViewRow> InRow, TViewModelPtr<IOutlinerExtension> InDataModel);
+	UE_API TSharedPtr<STrackLane> FindOrCreateParentLane(TViewModelPtr<IOutlinerExtension> InDataModel);
 
 	/** Gather the children from the specified node */
-	void OnGetChildren(TWeakViewModelPtr<IOutlinerExtension> InParent, TArray<TWeakViewModelPtr<IOutlinerExtension>>& OutChildren) const;
+	UE_API void OnGetChildren(TWeakViewModelPtr<IOutlinerExtension> InParent, TArray<TWeakViewModelPtr<IOutlinerExtension>>& OutChildren) const;
 
 	/** Generate a widget for the specified Node and Column */
-	TSharedRef<SWidget> GenerateWidgetForColumn(TViewModelPtr<IOutlinerExtension> InDataModel, const FName& ColumnId, const TSharedRef<SOutlinerViewRow>& Row) const;
+	UE_API TSharedRef<SWidget> GenerateWidgetForColumn(TViewModelPtr<IOutlinerExtension> InDataModel, const FName& ColumnId, const TSharedRef<SOutlinerViewRow>& Row) const;
 
 	/** Called when a node has been expanded or collapsed */
-	void OnExpansionChanged(TWeakViewModelPtr<IOutlinerExtension> InItem, bool bIsExpanded);
+	UE_API void OnExpansionChanged(TWeakViewModelPtr<IOutlinerExtension> InItem, bool bIsExpanded);
 
 	// Tree selection methods which must be overriden to maintain selection consistency with the rest of sequencer.
-	virtual void Private_UpdateParentHighlights() override;
-	virtual void Private_SetItemSelection( TWeakViewModelPtr<IOutlinerExtension> TheItem, bool bShouldBeSelected, bool bWasUserDirected = false ) override;
-	virtual void Private_ClearSelection() override;
-	virtual void Private_SelectRangeFromCurrentTo( TWeakViewModelPtr<IOutlinerExtension> InRangeSelectionEnd ) override;
-	virtual void Private_SignalSelectionChanged( ESelectInfo::Type SelectInfo ) override;
+	UE_API virtual void Private_UpdateParentHighlights() override;
+	UE_API virtual void Private_SetItemSelection( TWeakViewModelPtr<IOutlinerExtension> TheItem, bool bShouldBeSelected, bool bWasUserDirected = false ) override;
+	UE_API virtual void Private_ClearSelection() override;
+	UE_API virtual void Private_SelectRangeFromCurrentTo( TWeakViewModelPtr<IOutlinerExtension> InRangeSelectionEnd ) override;
+	UE_API virtual void Private_SignalSelectionChanged( ESelectInfo::Type SelectInfo ) override;
 
-	virtual void OnRightMouseButtonDown(const FPointerEvent& MouseEvent) override;
-	virtual void OnRightMouseButtonUp(const FPointerEvent& MouseEvent) override;
+	UE_API virtual void OnRightMouseButtonDown(const FPointerEvent& MouseEvent) override;
+	UE_API virtual void OnRightMouseButtonUp(const FPointerEvent& MouseEvent) override;
 
-	virtual FReply OnDragOver(const FGeometry& MyGeometry, const FDragDropEvent& DragDropEvent) override;
-	virtual FReply OnDrop(const FGeometry& MyGeometry, const FDragDropEvent& DragDropEvent) override;
+	UE_API virtual FReply OnDragOver(const FGeometry& MyGeometry, const FDragDropEvent& DragDropEvent) override;
+	UE_API virtual FReply OnDrop(const FGeometry& MyGeometry, const FDragDropEvent& DragDropEvent) override;
 
 public:
 
-	virtual FReply OnKeyDown(const FGeometry& MyGeometry, const FKeyEvent& InKeyEvent) override;
+	UE_API virtual FReply OnKeyDown(const FGeometry& MyGeometry, const FKeyEvent& InKeyEvent) override;
 
 protected:
 
 	// Private, unimplemented overloaded name for SetItemSelection to prevent external calls - use ForceSetItemSelection instead
-	void SetItemSelection();
+	UE_API void SetItemSelection();
 	// Private, unimplemented overloaded name for SetItemSelection to prevent external calls - use ForceSetItemSelection instead
-	void ClearSelection();
+	UE_API void ClearSelection();
 
-	void UpdateViewSelectionFromModel();
-	void UpdateModelSelectionFromView();
+	UE_API void UpdateViewSelectionFromModel();
+	UE_API void UpdateModelSelectionFromView();
 
-	void HandleTableViewScrolled(double InScrollOffset);
-	void UpdatePhysicalGeometry(bool bIsRefresh);
+	UE_API void HandleTableViewScrolled(double InScrollOffset);
+	UE_API void UpdatePhysicalGeometry(bool bIsRefresh);
 
 	/** Handles the context menu opening when right clicking on the tree view. */
-	TSharedPtr<SWidget> OnContextMenuOpening();
+	UE_API TSharedPtr<SWidget> OnContextMenuOpening();
 
-	void SetItemExpansionRecursive(TWeakViewModelPtr<IOutlinerExtension> InItem, bool bIsExpanded);
+	UE_API void SetItemExpansionRecursive(TWeakViewModelPtr<IOutlinerExtension> InItem, bool bIsExpanded);
 
 public:
 
@@ -205,7 +207,7 @@ public:
 	const TArray<FCachedGeometry>& GetAllVisibleNodes() const { return PhysicalNodes; }
 
 	/** Add a SOutlinerView object that should be modified or updated when this Treeview is updated */
-	void AddPinnedTreeView(TSharedPtr<SOutlinerView> PinnedTreeView);
+	UE_API void AddPinnedTreeView(TSharedPtr<SOutlinerView> PinnedTreeView);
 
 	/** Set a SOutlinerView object this Treeview is pinned to, for operations that should happen on the primary */
 	void SetPrimaryTreeView(TSharedPtr<SOutlinerView> InPrimaryTreeView) { PrimaryTreeView = InPrimaryTreeView; }
@@ -214,24 +216,24 @@ public:
 	void SetShowPinned(bool bShowPinned) { bShowPinnedNodes = bShowPinned; }
 
 	/** Updates the list of visible outliner columns and regenerates columns in the outliner view */
-	void SetOutlinerColumns(const TArray<TSharedPtr<IOutlinerColumn>>& InOutlinerColumns);
+	UE_API void SetOutlinerColumns(const TArray<TSharedPtr<IOutlinerColumn>>& InOutlinerColumns);
 
 protected:
 
 	/** Linear, sorted array of nodes that we currently have generated widgets for */
 	TArray<FCachedGeometry> PhysicalNodes;
 
-	int32 CreateOutlinerColumnsForGroup(int32 ColumnIndex, EOutlinerColumnGroup Group);
+	UE_API int32 CreateOutlinerColumnsForGroup(int32 ColumnIndex, EOutlinerColumnGroup Group);
 
 	/** Populate the map of column definitions, and add relevant columns to the header row. Must be called when outliner columns change */
-	void UpdateOutlinerColumns();
+	UE_API void UpdateOutlinerColumns();
 
 	/** Insert a separator column at the specified column index, with a unique identifier */
-	void InsertSeparatorColumn(int32 InsertIndex, int32 SeparatorID);
+	UE_API void InsertSeparatorColumn(int32 InsertIndex, int32 SeparatorID);
 
-	FReply OnDragRow(const FGeometry& InGeometry, const FPointerEvent& InPointerEvent, TSharedRef<SOutlinerViewRow> InRow);
+	UE_API FReply OnDragRow(const FGeometry& InGeometry, const FPointerEvent& InPointerEvent, TSharedRef<SOutlinerViewRow> InRow);
 
-	FString OnItemToString_Debug(TWeakViewModelPtr<IOutlinerExtension> InWeakModel);
+	UE_API FString OnItemToString_Debug(TWeakViewModelPtr<IOutlinerExtension> InWeakModel);
 
 protected:
 
@@ -287,3 +289,4 @@ protected:
 
 } // namespace UE::Sequencer
 
+#undef UE_API
diff --git a/Engine/Source/Editor/SequencerCore/Public/MVVM/Views/SOutlinerViewRow.h b/Engine/Source/Editor/SequencerCore/Public/MVVM/Views/SOutlinerViewRow.h
--- a/Engine/Source/Editor/SequencerCore/Public/MVVM/Views/SOutlinerViewRow.h
+++ b/Engine/Source/Editor/SequencerCore/Public/MVVM/Views/SOutlinerViewRow.h
@@ -6,6 +6,8 @@
 #include "SequencerCoreFwd.h"
 #include "MVVM/Extensions/IOutlinerExtension.h"
 
+#define UE_API SEQUENCERCORE_API
+
 namespace UE::Sequencer
 {
 
@@ -14,7 +16,7 @@ class STrackLane;
 /**
  * Widget that represents a row in the sequencer's tree control.
  */
-class SEQUENCERCORE_API SOutlinerViewRow
+class SOutlinerViewRow
 	: public ISequencerTreeViewRow
 {
 public:
@@ -36,13 +38,13 @@ public:
 	SLATE_END_ARGS()
 
 	/** Construct function for this widget */
-	void Construct(const FArguments& InArgs, const TSharedRef<STableViewBase>& InOwnerTableView, TWeakViewModelPtr<IOutlinerExtension> InDataModel);
+	UE_API void Construct(const FArguments& InArgs, const TSharedRef<STableViewBase>& InOwnerTableView, TWeakViewModelPtr<IOutlinerExtension> InDataModel);
 
 	/** Destroy this widget */
-	~SOutlinerViewRow();
+	UE_API ~SOutlinerViewRow();
 
 	/** Get the model to which this row relates */
-	TViewModelPtr<IOutlinerExtension> GetDataModel() const;
+	UE_API TViewModelPtr<IOutlinerExtension> GetDataModel() const;
 
 	/**
 	 * Gets the track lane we relate to
@@ -50,50 +52,50 @@ public:
 	 *							 by our own outliner item view-model 
 	 * @return The track lane for this row
 	 */
-	TSharedPtr<STrackLane> GetTrackLane(bool bOnlyOwnTrackLane = false) const;
+	UE_API TSharedPtr<STrackLane> GetTrackLane(bool bOnlyOwnTrackLane = false) const;
 
 	/**
 	 * Adds a reference to track lane, either because:
 	 *  - It is a parent row's track lane that we want to be kept alive when the parent row disappears out of view, or
 	 *  - It is a track lane that was created by our outliner view-model
 	 */
-	void SetTrackLane(const TSharedPtr<STrackLane>& InTrackLane);
+	UE_API void SetTrackLane(const TSharedPtr<STrackLane>& InTrackLane);
 
 	/** Whether the underlying data model is selectable */
-	bool IsSelectable() const;
+	UE_API bool IsSelectable() const;
 
-	virtual const FSlateBrush* GetBorder() const override;
+	UE_API virtual const FSlateBrush* GetBorder() const override;
 
-	virtual void Tick(const FGeometry& AllottedGeometry, const double InCurrentTime, const float InDeltaTime) override;
+	UE_API virtual void Tick(const FGeometry& AllottedGeometry, const double InCurrentTime, const float InDeltaTime) override;
 	
-	virtual void OnArrangeChildren(const FGeometry& AllottedGeometry, FArrangedChildren& ArrangedChildren) const override;
+	UE_API virtual void OnArrangeChildren(const FGeometry& AllottedGeometry, FArrangedChildren& ArrangedChildren) const override;
 
-	virtual int32 OnPaint(const FPaintArgs& Args, const FGeometry& AllottedGeometry, const FSlateRect& MyCullingRect, FSlateWindowElementList& OutDrawElements, int32 LayerId, const FWidgetStyle& InWidgetStyle, bool bParentEnabled) const override;
+	UE_API virtual int32 OnPaint(const FPaintArgs& Args, const FGeometry& AllottedGeometry, const FSlateRect& MyCullingRect, FSlateWindowElementList& OutDrawElements, int32 LayerId, const FWidgetStyle& InWidgetStyle, bool bParentEnabled) const override;
 
-	virtual int32 OnPaintDropIndicator(EItemDropZone InItemDropZone, const FPaintArgs& Args, const FGeometry& AllottedGeometry, const FSlateRect& MyCullingRect, FSlateWindowElementList& OutDrawElements, int32 LayerId, const FWidgetStyle& InWidgetStyle, bool bParentEnabled) const override;
+	UE_API virtual int32 OnPaintDropIndicator(EItemDropZone InItemDropZone, const FPaintArgs& Args, const FGeometry& AllottedGeometry, const FSlateRect& MyCullingRect, FSlateWindowElementList& OutDrawElements, int32 LayerId, const FWidgetStyle& InWidgetStyle, bool bParentEnabled) const override;
 
-	virtual void ConstructChildren(ETableViewMode::Type InOwnerTableMode, const TAttribute<FMargin>& InPadding, const TSharedRef<SWidget>& InContent) override;
+	UE_API virtual void ConstructChildren(ETableViewMode::Type InOwnerTableMode, const TAttribute<FMargin>& InPadding, const TSharedRef<SWidget>& InContent) override;
 
-	virtual void OnMouseEnter(const FGeometry& MyGeometry, const FPointerEvent& MouseEvent) override;
+	UE_API virtual void OnMouseEnter(const FGeometry& MyGeometry, const FPointerEvent& MouseEvent) override;
 
-	virtual void OnMouseLeave(const FPointerEvent& MouseEvent) override;
+	UE_API virtual void OnMouseLeave(const FPointerEvent& MouseEvent) override;
 
-	virtual bool IsColumnVisible(const FName& InColumnName) const override;
+	UE_API virtual bool IsColumnVisible(const FName& InColumnName) const override;
 
 	/** Called whenever a drag is detected by the tree view. */
-	FReply OnDragDetected(const FGeometry& InGeometry, const FPointerEvent& InPointerEvent) override;
+	UE_API FReply OnDragDetected(const FGeometry& InGeometry, const FPointerEvent& InPointerEvent) override;
 
 	/** Called to determine whether a current drag operation is valid for this row. */
-	TOptional<EItemDropZone> OnCanAcceptDrop( const FDragDropEvent& DragDropEvent, EItemDropZone ItemDropZone, TWeakViewModelPtr<IOutlinerExtension> InDataModel);
+	UE_API TOptional<EItemDropZone> OnCanAcceptDrop( const FDragDropEvent& DragDropEvent, EItemDropZone ItemDropZone, TWeakViewModelPtr<IOutlinerExtension> InDataModel);
 
 	/** Called to complete a drag and drop onto this drop. */
-	FReply OnAcceptDrop( const FDragDropEvent& DragDropEvent, EItemDropZone ItemDropZone, TWeakViewModelPtr<IOutlinerExtension> InDataModel);
+	UE_API FReply OnAcceptDrop( const FDragDropEvent& DragDropEvent, EItemDropZone ItemDropZone, TWeakViewModelPtr<IOutlinerExtension> InDataModel);
 
 protected:
 
-	TSharedPtr<SWidget> GenerateWidgetForColumn(const FName& ColumnId);
+	UE_API TSharedPtr<SWidget> GenerateWidgetForColumn(const FName& ColumnId);
 
-	FVector2D ComputeDesiredSize(float LayoutScaleMultiplier) const override;
+	UE_API FVector2D ComputeDesiredSize(float LayoutScaleMultiplier) const override;
 
 protected:
 
@@ -122,3 +124,4 @@ protected:
 
 } // namespace UE::Sequencer
 
+#undef UE_API
diff --git a/Engine/Source/Editor/SequencerCore/Public/MVVM/Views/STrackAreaLaneView.h b/Engine/Source/Editor/SequencerCore/Public/MVVM/Views/STrackAreaLaneView.h
--- a/Engine/Source/Editor/SequencerCore/Public/MVVM/Views/STrackAreaLaneView.h
+++ b/Engine/Source/Editor/SequencerCore/Public/MVVM/Views/STrackAreaLaneView.h
@@ -7,13 +7,15 @@
 #include "MVVM/ViewModelPtr.h"
 #include "MVVM/Extensions/ITrackLaneExtension.h"
 
+#define UE_API SEQUENCERCORE_API
+
 
 namespace UE::Sequencer
 {
 
 class STrackAreaView;
 
-class SEQUENCERCORE_API STrackAreaLaneView
+class STrackAreaLaneView
 	: public SCompoundWidget
 	, public ITrackLaneWidget
 {
@@ -22,13 +24,13 @@ public:
 		SLATE_DEFAULT_SLOT(FArguments, Content)
 	SLATE_END_ARGS()
 
-	void Construct(const FArguments& InArgs, const FViewModelPtr& InViewModel, TSharedPtr<STrackAreaView> InTrackAreaView);
+	UE_API void Construct(const FArguments& InArgs, const FViewModelPtr& InViewModel, TSharedPtr<STrackAreaView> InTrackAreaView);
 
 	/*~ ITrackLaneWidget */
-	virtual TSharedRef<const SWidget> AsWidget() const;
-	virtual FTrackLaneScreenAlignment GetAlignment(const ITrackLaneWidgetSpace& InScreenSpace, const FGeometry& InParentGeometry) const;
+	UE_API virtual TSharedRef<const SWidget> AsWidget() const;
+	UE_API virtual FTrackLaneScreenAlignment GetAlignment(const ITrackLaneWidgetSpace& InScreenSpace, const FGeometry& InParentGeometry) const;
 
-	FTimeToPixel GetRelativeTimeToPixel() const;
+	UE_API FTimeToPixel GetRelativeTimeToPixel() const;
 
 protected:
 
@@ -38,3 +40,5 @@ protected:
 };
 
 } // namespace UE::Sequencer
+
+#undef UE_API
diff --git a/Engine/Source/Editor/SequencerCore/Public/MVVM/Views/STrackAreaView.h b/Engine/Source/Editor/SequencerCore/Public/MVVM/Views/STrackAreaView.h
--- a/Engine/Source/Editor/SequencerCore/Public/MVVM/Views/STrackAreaView.h
+++ b/Engine/Source/Editor/SequencerCore/Public/MVVM/Views/STrackAreaView.h
@@ -22,6 +22,8 @@
 #include "Widgets/DeclarativeSyntaxSupport.h"
 #include "Widgets/SPanel.h"
 
+#define UE_API SEQUENCERCORE_API
+
 class FArrangedChildren;
 class FChildren;
 class FDragDropEvent;
@@ -74,7 +76,7 @@ struct FTrackAreaViewLayers
 /**
  * The area where tracks( rows of sections ) are displayed.
  */
-class SEQUENCERCORE_API STrackAreaView
+class STrackAreaView
 	: public SPanel
 	, public ITrackLaneWidgetSpace
 {
@@ -86,32 +88,32 @@ public:
 	}
 	SLATE_END_ARGS()
 
-	STrackAreaView();
-	~STrackAreaView();
+	UE_API STrackAreaView();
+	UE_API ~STrackAreaView();
 
 	/** Construct this widget */
-	void Construct(const FArguments& InArgs, TWeakPtr<FTrackAreaViewModel> InWeakViewModel);
+	UE_API void Construct(const FArguments& InArgs, TWeakPtr<FTrackAreaViewModel> InWeakViewModel);
 
-	TSharedPtr<FTrackAreaViewModel> GetViewModel() const;
+	UE_API TSharedPtr<FTrackAreaViewModel> GetViewModel() const;
 	
-	void SetVirtualPosition(float InVirtualTop);
+	UE_API void SetVirtualPosition(float InVirtualTop);
 
-	TSharedPtr<FTimeToPixel> GetFallbackTimeToPixel() const;
-	TSharedPtr<FTimeToPixel> GetTimeToPixel(const FGuid& InViewSpaceID = FGuid()) const;
+	UE_API TSharedPtr<FTimeToPixel> GetFallbackTimeToPixel() const;
+	UE_API TSharedPtr<FTimeToPixel> GetTimeToPixel(const FGuid& InViewSpaceID = FGuid()) const;
 
 public:
 
 	/** Empty the track area */
-	void Empty();
+	UE_API void Empty();
 
 	/** Add a new track slot to this area for the given node. The slot will be automatically cleaned up when all external references to the supplied slot are removed. */
-	void AddTrackSlot(const TViewModelPtr<IOutlinerExtension>& InDataModel, const TSharedPtr<STrackLane>& InSlot);
+	UE_API void AddTrackSlot(const TViewModelPtr<IOutlinerExtension>& InDataModel, const TSharedPtr<STrackLane>& InSlot);
 
 	/** Attempt to find an existing slot relating to the given node */
-	TSharedPtr<STrackLane> FindTrackSlot(const TViewModelPtr<IOutlinerExtension>& InDataModel);
+	UE_API TSharedPtr<STrackLane> FindTrackSlot(const TViewModelPtr<IOutlinerExtension>& InDataModel);
 
 	/** Assign a tree view to this track area. */
-	void SetOutliner(const TSharedPtr<SOutlinerView>& InOutliner);
+	UE_API void SetOutliner(const TSharedPtr<SOutlinerView>& InOutliner);
 	TWeakPtr<SOutlinerView> GetOutliner() const { return WeakOutliner.Pin(); }
 
 	/** Set whether this TrackArea should show only pinned nodes or only non-pinned nodes  */
@@ -122,36 +124,36 @@ public:
 	void SetIsPinned(bool bInIsPinned) { bIsPinned = bInIsPinned; }
 	bool IsPinned() const { return bIsPinned; }
 
-	const FTrackAreaViewLayers& GetPaintLayers() const;
-	static FLinearColor BlendDefaultTrackColor(FLinearColor InColor);
+	UE_API const FTrackAreaViewLayers& GetPaintLayers() const;
+	static UE_API FLinearColor BlendDefaultTrackColor(FLinearColor InColor);
 
 public:
 
 	/*~ SWidget interface */
-	FReply OnMouseButtonDown( const FGeometry& MyGeometry, const FPointerEvent& MouseEvent ) override;
-	FReply OnMouseButtonUp( const FGeometry& MyGeometry, const FPointerEvent& MouseEvent ) override;
-	FReply OnMouseButtonDoubleClick( const FGeometry& MyGeometry, const FPointerEvent& MouseEvent ) override;
-	FReply OnMouseMove( const FGeometry& MyGeometry, const FPointerEvent& MouseEvent ) override;
-	FReply OnMouseWheel( const FGeometry& MyGeometry, const FPointerEvent& MouseEvent ) override;
-	void OnMouseEnter(const FGeometry& MyGeometry, const FPointerEvent& MouseEvent) override;
-	void OnMouseLeave(const FPointerEvent& MouseEvent) override;
-	int32 OnPaint(const FPaintArgs& Args, const FGeometry& AllottedGeometry, const FSlateRect& MyCullingRect, FSlateWindowElementList& OutDrawElements, int32 LayerId, const FWidgetStyle& InWidgetStyle, bool bParentEnabled ) const override;
-	void OnMouseCaptureLost(const FCaptureLostEvent& CaptureLostEvent) override;
-	FCursorReply OnCursorQuery( const FGeometry& MyGeometry, const FPointerEvent& CursorEvent ) const override;
-	void Tick( const FGeometry& AllottedGeometry, const double InCurrentTime, const float InDeltaTime ) override;
-	void OnArrangeChildren( const FGeometry& AllottedGeometry, FArrangedChildren& ArrangedChildren ) const override;
-	FVector2D ComputeDesiredSize(float) const override;
-	FChildren* GetChildren() override;
-	void OnDragEnter(const FGeometry& MyGeometry, const FDragDropEvent& DragDropEvent) override;
-	void OnDragLeave(const FDragDropEvent& DragDropEvent) override;
-	FReply OnDragOver(const FGeometry& MyGeometry, const FDragDropEvent& DragDropEvent) override;
-	FReply OnDrop(const FGeometry& MyGeometry, const FDragDropEvent& DragDropEvent) override;
+	UE_API FReply OnMouseButtonDown( const FGeometry& MyGeometry, const FPointerEvent& MouseEvent ) override;
+	UE_API FReply OnMouseButtonUp( const FGeometry& MyGeometry, const FPointerEvent& MouseEvent ) override;
+	UE_API FReply OnMouseButtonDoubleClick( const FGeometry& MyGeometry, const FPointerEvent& MouseEvent ) override;
+	UE_API FReply OnMouseMove( const FGeometry& MyGeometry, const FPointerEvent& MouseEvent ) override;
+	UE_API FReply OnMouseWheel( const FGeometry& MyGeometry, const FPointerEvent& MouseEvent ) override;
+	UE_API void OnMouseEnter(const FGeometry& MyGeometry, const FPointerEvent& MouseEvent) override;
+	UE_API void OnMouseLeave(const FPointerEvent& MouseEvent) override;
+	UE_API int32 OnPaint(const FPaintArgs& Args, const FGeometry& AllottedGeometry, const FSlateRect& MyCullingRect, FSlateWindowElementList& OutDrawElements, int32 LayerId, const FWidgetStyle& InWidgetStyle, bool bParentEnabled ) const override;
+	UE_API void OnMouseCaptureLost(const FCaptureLostEvent& CaptureLostEvent) override;
+	UE_API FCursorReply OnCursorQuery( const FGeometry& MyGeometry, const FPointerEvent& CursorEvent ) const override;
+	UE_API void Tick( const FGeometry& AllottedGeometry, const double InCurrentTime, const float InDeltaTime ) override;
+	UE_API void OnArrangeChildren( const FGeometry& AllottedGeometry, FArrangedChildren& ArrangedChildren ) const override;
+	UE_API FVector2D ComputeDesiredSize(float) const override;
+	UE_API FChildren* GetChildren() override;
+	UE_API void OnDragEnter(const FGeometry& MyGeometry, const FDragDropEvent& DragDropEvent) override;
+	UE_API void OnDragLeave(const FDragDropEvent& DragDropEvent) override;
+	UE_API FReply OnDragOver(const FGeometry& MyGeometry, const FDragDropEvent& DragDropEvent) override;
+	UE_API FReply OnDrop(const FGeometry& MyGeometry, const FDragDropEvent& DragDropEvent) override;
 	virtual bool SupportsKeyboardFocus() const override { return true; }
-	virtual FReply OnKeyDown(const FGeometry& MyGeometry, const FKeyEvent& InKeyEvent) override;
-	virtual FReply OnKeyUp(const FGeometry& MyGeometry, const FKeyEvent& InKeyEvent) override;
+	UE_API virtual FReply OnKeyDown(const FGeometry& MyGeometry, const FKeyEvent& InKeyEvent) override;
+	UE_API virtual FReply OnKeyUp(const FGeometry& MyGeometry, const FKeyEvent& InKeyEvent) override;
 
 	/*~ Begin ITrackLaneWidgetSpace Interface */
-	virtual FTimeToPixel GetScreenSpace(const FGuid& InViewSpaceID = FGuid()) const override;
+	UE_API virtual FTimeToPixel GetScreenSpace(const FGuid& InViewSpaceID = FGuid()) const override;
 	/*~ End ITrackLaneWidgetSpace Interface */
 
 protected:
@@ -159,12 +161,12 @@ protected:
 	virtual void OnResized(const FVector2D& OldSize, const FVector2D& NewSize){}
 
 	/** Update any hover state required for the track area */
-	virtual void UpdateHoverStates( const FGeometry& MyGeometry, const FPointerEvent& MouseEvent );
+	UE_API virtual void UpdateHoverStates( const FGeometry& MyGeometry, const FPointerEvent& MouseEvent );
 
 private:
 
-	void HandleViewSpaceAdded(const FGuid& InID, TSharedPtr<FTrackAreaViewSpace> ViewSpace);
-	void HandleViewSpaceRemoved(const FGuid& InID, TSharedPtr<FTrackAreaViewSpace> ViewSpace);
+	UE_API void HandleViewSpaceAdded(const FGuid& InID, TSharedPtr<FTrackAreaViewSpace> ViewSpace);
+	UE_API void HandleViewSpaceRemoved(const FGuid& InID, TSharedPtr<FTrackAreaViewSpace> ViewSpace);
 
 private:
 
@@ -214,3 +216,4 @@ protected:
 } // namespace Sequencer
 } // namespace UE
 
+#undef UE_API
