Sequencer Scripting: GetMuteNodes/SetMuteNodes, GetSoloNodes/SetSoloNodes, GetDeactivatedNodes/SetDeactivatedNodes, GetLockedNodes, SetLockedNodes, GetPinnedNodes, SetPinnedNodes, AddBindingTag, RemoveBindingTag, GetBindingTags

diff --git a/Engine/Source/Editor/Sequencer/Private/MVVM/ViewModels/OutlinerItemModel.cpp b/Engine/Source/Editor/Sequencer/Private/MVVM/ViewModels/OutlinerItemModel.cpp
--- a/Engine/Source/Editor/Sequencer/Private/MVVM/ViewModels/OutlinerItemModel.cpp
+++ b/Engine/Source/Editor/Sequencer/Private/MVVM/ViewModels/OutlinerItemModel.cpp
@@ -1189,6 +1189,9 @@ void FEvaluableOutlinerItemModel::SetIsSoloed(bool bIsSoloed)
 			return;
 		}
 
+		// This shouldn't modify because it's not saved with the asset. But it still needs to mark the movie scene as changed, so that the UI will refresh
+		MovieScene->MarkAsChanged();
+
 		TArray<FString>& SoloNodes = MovieScene->GetSoloNodes();
 
 		// Should be always called on the main thread, but thread_local for safety
@@ -1247,6 +1250,9 @@ void FEvaluableOutlinerItemModel::SetIsMuted(bool bIsMuted)
 			return;
 		}
 
+		// This shouldn't modify because it's not saved with the asset. But it still needs to mark the movie scene as changed, so that the UI will refresh
+		MovieScene->MarkAsChanged();
+
 		TArray<FString>& MuteNodes = MovieScene->GetMuteNodes();
 
 		// Should be always called on the main thread, but thread_local for safety
diff --git a/Engine/Source/Editor/Sequencer/Private/Scripting/SequencerModuleOutlinerScriptingObject.cpp b/Engine/Source/Editor/Sequencer/Private/Scripting/SequencerModuleOutlinerScriptingObject.cpp
--- a/Engine/Source/Editor/Sequencer/Private/Scripting/SequencerModuleOutlinerScriptingObject.cpp
+++ b/Engine/Source/Editor/Sequencer/Private/Scripting/SequencerModuleOutlinerScriptingObject.cpp
@@ -3,16 +3,21 @@
 #include "Scripting/SequencerModuleOutlinerScriptingObject.h"
 #include "Scripting/ViewModelScriptingStruct.h"
 
+#include "MVVM/ViewModels/SequencerEditorViewModel.h"
 #include "MVVM/ViewModels/SequencerOutlinerViewModel.h"
 #include "MVVM/ViewModels/ViewModel.h"
 #include "MVVM/ViewModels/ViewModelIterators.h"
+#include "MVVM/ViewModels/ObjectBindingModel.h"
 #include "MVVM/ViewModels/OutlinerViewModel.h"
 #include "MVVM/ViewModels/EditorViewModel.h"
 #include "MVVM/ViewModelPtr.h"
 
 #include "Containers/Array.h"
 #include "Misc/FrameNumber.h"
+#include "MovieScene.h"
+#include "ObjectBindingTagCache.h"
 
+#include "Sequencer.h"
 #include "SequencerCommonHelpers.h"
 
 #define LOCTEXT_NAMESPACE "SequencerModuleOutlinerScriptingObject"
@@ -102,4 +107,147 @@ FFrameNumber USequencerModuleOutlinerScriptingObject::GetPreviousKey(const TArra
 	return Outliner->GetPreviousKey(Nodes, FrameNumber, TimeUnit, Range);
 }
 
+void USequencerModuleOutlinerScriptingObject::AddBindingTags(const TArray<FSequencerViewModelScriptingStruct>& InNodes, const TArray<FName>& TagNames)
+{
+	using namespace UE::Sequencer;
+
+	TSharedPtr<FSequencer> Sequencer = GetSequencer();
+	if (!Sequencer)
+	{
+		FFrame::KismetExecutionMessage(*LOCTEXT("SequencerInvalid", "Sequencer is no longer valid.").ToString(), ELogVerbosity::Error);
+		return;
+	}
+
+	UMovieSceneSequence* Sequence = Sequencer->GetFocusedMovieSceneSequence();
+	UMovieScene* MovieScene = Sequence ? Sequence->GetMovieScene() : nullptr;
+
+	if (!MovieScene)
+	{
+		FFrame::KismetExecutionMessage(*LOCTEXT("MovieSceneInvalid", "Movie Scene is not an object binding.").ToString(), ELogVerbosity::Error);
+		return;
+	}
+
+	for (const FSequencerViewModelScriptingStruct& InNode : InNodes)
+	{
+		TViewModelPtr<FObjectBindingModel> ObjectBindingModel = CastViewModel<FObjectBindingModel>(InNode.WeakViewModel.Pin());
+		if (!ObjectBindingModel)
+		{
+			FFrame::KismetExecutionMessage(*LOCTEXT("NodeInvalid", "Node is not an object binding.").ToString(), ELogVerbosity::Error);
+			return;
+		}
+
+		const FGuid& ObjectID = ObjectBindingModel->GetObjectGuid();
+		FMovieSceneSequenceIDRef SequenceID = Sequencer->GetFocusedTemplateID();
+
+		for (const FName& TagName : TagNames)
+		{
+			UE::MovieScene::FFixedObjectBindingID BindingID(ObjectID, SequenceID);
+			MovieScene->TagBinding(TagName, BindingID);
+		}
+	}
+}
+
+void USequencerModuleOutlinerScriptingObject::RemoveBindingTags(const TArray<FSequencerViewModelScriptingStruct>& InNodes, const TArray<FName>& TagNames)
+{
+	using namespace UE::Sequencer;
+
+	TSharedPtr<FSequencer> Sequencer = GetSequencer();
+	if (!Sequencer)
+	{
+		FFrame::KismetExecutionMessage(*LOCTEXT("SequencerInvalid", "Sequencer is no longer valid.").ToString(), ELogVerbosity::Error);
+		return;
+	}
+
+	UMovieSceneSequence* Sequence = Sequencer->GetFocusedMovieSceneSequence();
+	UMovieScene* MovieScene = Sequence ? Sequence->GetMovieScene() : nullptr;
+
+	if (!MovieScene)
+	{
+		FFrame::KismetExecutionMessage(*LOCTEXT("MovieSceneInvalid", "Movie Scene is not an object binding.").ToString(), ELogVerbosity::Error);
+		return;
+	}	
+
+	for (const FSequencerViewModelScriptingStruct& InNode : InNodes)
+	{
+		TViewModelPtr<FObjectBindingModel> ObjectBindingModel = CastViewModel<FObjectBindingModel>(InNode.WeakViewModel.Pin());
+		if (!ObjectBindingModel)
+		{
+			FFrame::KismetExecutionMessage(*LOCTEXT("NodeInvalid", "Node is not an object binding.").ToString(), ELogVerbosity::Error);
+			return;
+		}
+
+		const FGuid& ObjectID = ObjectBindingModel->GetObjectGuid();
+		FMovieSceneSequenceIDRef SequenceID = Sequencer->GetFocusedTemplateID();
+
+		for (const FName& TagName : TagNames)
+		{
+			UE::MovieScene::FFixedObjectBindingID BindingID(ObjectID, SequenceID);
+			MovieScene->UntagBinding(TagName, BindingID);
+		}
+	}
+}
+
+TArray<FName> USequencerModuleOutlinerScriptingObject::GetBindingTags(FSequencerViewModelScriptingStruct InNode)
+{
+	using namespace UE::Sequencer;
+
+	TArray<FName> Tags;
+
+	TSharedPtr<FSequencer> Sequencer = GetSequencer();
+	if (!Sequencer)
+	{
+		FFrame::KismetExecutionMessage(*LOCTEXT("SequencerInvalid", "Sequencer is no longer valid.").ToString(), ELogVerbosity::Error);
+		return Tags;
+	}
+
+	UMovieSceneSequence* Sequence = Sequencer->GetFocusedMovieSceneSequence();
+	UMovieScene* MovieScene = Sequence ? Sequence->GetMovieScene() : nullptr;
+
+	if (!MovieScene)
+	{
+		FFrame::KismetExecutionMessage(*LOCTEXT("MovieSceneInvalid", "Movie Scene is not an object binding.").ToString(), ELogVerbosity::Error);
+		return Tags;
+	}
+
+	TViewModelPtr<FObjectBindingModel> ObjectBindingModel = CastViewModel<FObjectBindingModel>(InNode.WeakViewModel.Pin());
+	if (!ObjectBindingModel)
+	{
+		FFrame::KismetExecutionMessage(*LOCTEXT("NodeInvalid", "Node is not an object binding.").ToString(), ELogVerbosity::Error);
+		return Tags;
+	}
+
+	const FGuid& ObjectID = ObjectBindingModel->GetObjectGuid();
+	FMovieSceneSequenceIDRef SequenceID = Sequencer->GetFocusedTemplateID();
+
+	UE::MovieScene::FFixedObjectBindingID BindingID(ObjectID, SequenceID);
+	for (auto It = Sequencer->GetObjectBindingTagCache()->IterateTags(BindingID); It; ++It)
+	{
+		Tags.Add(It.Value());
+	}
+
+	return Tags;
+}
+
+TSharedPtr<FSequencer> USequencerModuleOutlinerScriptingObject::GetSequencer()
+{
+	using namespace UE::Sequencer;
+
+	TViewModelPtr<FSequencerOutlinerViewModel> Outliner = CastViewModel<FSequencerOutlinerViewModel>(WeakOutliner.Pin());
+	if (!Outliner)
+	{
+		FFrame::KismetExecutionMessage(*LOCTEXT("SequencerOutlinerInvalid", "Sequencer Outliner is no longer valid.").ToString(), ELogVerbosity::Error);
+		return nullptr;
+	}
+
+	TViewModelPtr<FSequencerEditorViewModel> Editor = CastViewModel<FSequencerEditorViewModel>(Outliner->GetEditor());
+	if (!Editor)
+	{
+		FFrame::KismetExecutionMessage(*LOCTEXT("SequencerEditorInvalid", "Sequencer Editor is no longer valid.").ToString(), ELogVerbosity::Error);
+		return nullptr;
+	}
+
+	TSharedPtr<FSequencer> Sequencer = Editor->GetSequencerImpl();
+	return Sequencer;
+}
+
 #undef LOCTEXT_NAMESPACE
diff --git a/Engine/Source/Editor/Sequencer/Private/Scripting/SequencerModuleOutlinerScriptingObject.h b/Engine/Source/Editor/Sequencer/Private/Scripting/SequencerModuleOutlinerScriptingObject.h
--- a/Engine/Source/Editor/Sequencer/Private/Scripting/SequencerModuleOutlinerScriptingObject.h
+++ b/Engine/Source/Editor/Sequencer/Private/Scripting/SequencerModuleOutlinerScriptingObject.h
@@ -12,6 +12,7 @@
 struct FFrameNumber;
 struct FSequencerViewModelScriptingStruct;
 
+class FSequencer;
 class UMovieSceneSection;
 
 UCLASS()
@@ -29,4 +30,19 @@ public:
 
 	UFUNCTION(BlueprintCallable, Category = "Sequencer Editor")
 	FFrameNumber GetPreviousKey(const TArray<FSequencerViewModelScriptingStruct>& InNodes, FFrameNumber FrameNumber, EMovieSceneTimeUnit TimeUnit = EMovieSceneTimeUnit::DisplayRate) const;
+
+
+	// Tags
+	UFUNCTION(BlueprintCallable, Category = "Sequencer Editor")
+	void AddBindingTags(const TArray<FSequencerViewModelScriptingStruct>& InNodes, const TArray<FName>& TagNames);
+
+	UFUNCTION(BlueprintCallable, Category = "Sequencer Editor")
+	void RemoveBindingTags(const TArray<FSequencerViewModelScriptingStruct>& InNodes, const TArray<FName>& TagNames);
+
+	UFUNCTION(BlueprintCallable, Category = "Sequencer Editor")
+	TArray<FName> GetBindingTags(FSequencerViewModelScriptingStruct InNode);
+
+private:
+
+	TSharedPtr<FSequencer> GetSequencer();
 };
\ No newline at end of file
diff --git a/Engine/Source/Editor/SequencerCore/Private/Scripting/OutlinerScriptingObject.cpp b/Engine/Source/Editor/SequencerCore/Private/Scripting/OutlinerScriptingObject.cpp
--- a/Engine/Source/Editor/SequencerCore/Private/Scripting/OutlinerScriptingObject.cpp
+++ b/Engine/Source/Editor/SequencerCore/Private/Scripting/OutlinerScriptingObject.cpp
@@ -9,8 +9,13 @@
 #include "MVVM/ViewModels/EditorViewModel.h"
 #include "MVVM/ViewModelPtr.h"
 
+#include "MVVM/Extensions/IDeactivatableExtension.h"
+#include "MVVM/Extensions/IMutableExtension.h"
+#include "MVVM/Extensions/ILockableExtension.h"
 #include "MVVM/Extensions/IOutlinerExtension.h"
+#include "MVVM/Extensions/IPinnableExtension.h"
 #include "MVVM/Extensions/ISelectableExtension.h"
+#include "MVVM/Extensions/ISoloableExtension.h"
 
 #include "MVVM/Selection/SequencerCoreSelection.h"
 #include "MVVM/Selection/SequencerOutlinerSelection.h"
@@ -141,4 +146,189 @@ void USequencerOutlinerScriptingObject::SetSelection(const TArray<FSequencerView
 	}
 }
 
+TArray<FSequencerViewModelScriptingStruct> USequencerOutlinerScriptingObject::GetMuteNodes()
+{
+	TArray<FSequencerViewModelScriptingStruct> MutedNodes;
+
+	using namespace UE::Sequencer;
+	TViewModelPtr<FOutlinerViewModel> Outliner = WeakOutliner.Pin();
+	if (!Outliner)
+	{
+		FFrame::KismetExecutionMessage(*LOCTEXT("OutlinerInvalid", "Outliner is no longer valid.").ToString(), ELogVerbosity::Error);
+		return MutedNodes;
+	}
+
+	for (const TViewModelPtr<IMutableExtension>& Extension : Outliner->GetRootItem()->GetDescendantsOfType<IMutableExtension>(true))
+	{
+		if (Extension.IsValid() && Extension->IsMuted())
+		{
+			MutedNodes.Add(FSequencerViewModelScriptingStruct(Extension.AsModel()));
+		}
+	}
+
+	return MutedNodes;
+}
+
+void USequencerOutlinerScriptingObject::SetMuteNodes(const TArray<FSequencerViewModelScriptingStruct>& InNodes, bool bInMuted)
+{
+	using namespace UE::Sequencer;
+
+	for (const FSequencerViewModelScriptingStruct& Item : InNodes)
+	{
+		TViewModelPtr<IMutableExtension> MutableItem = Item.WeakViewModel.ImplicitPin();
+		if (MutableItem)
+		{
+			MutableItem->SetIsMuted(bInMuted);
+		}
+	}
+}
+
+TArray<FSequencerViewModelScriptingStruct> USequencerOutlinerScriptingObject::GetSoloNodes()
+{
+	TArray<FSequencerViewModelScriptingStruct> SoloNodes;
+
+	using namespace UE::Sequencer;
+	TViewModelPtr<FOutlinerViewModel> Outliner = WeakOutliner.Pin();
+	if (!Outliner)
+	{
+		FFrame::KismetExecutionMessage(*LOCTEXT("OutlinerInvalid", "Outliner is no longer valid.").ToString(), ELogVerbosity::Error);
+		return SoloNodes;
+	}
+
+	for (const TViewModelPtr<ISoloableExtension>& Extension : Outliner->GetRootItem()->GetDescendantsOfType<ISoloableExtension>(true))
+	{
+		if (Extension.IsValid() && Extension->IsSolo())
+		{
+			SoloNodes.Add(FSequencerViewModelScriptingStruct(Extension.AsModel()));
+		}
+	}
+
+	return SoloNodes;
+}
+
+void USequencerOutlinerScriptingObject::SetSoloNodes(const TArray<FSequencerViewModelScriptingStruct>& InNodes, bool bInSoloed)
+{
+	using namespace UE::Sequencer;
+
+	for (const FSequencerViewModelScriptingStruct& Item : InNodes)
+	{
+		TViewModelPtr<ISoloableExtension> SoloableItem = Item.WeakViewModel.ImplicitPin();
+		if (SoloableItem)
+		{
+			SoloableItem->SetIsSoloed(bInSoloed);
+		}
+	}
+}
+
+TArray<FSequencerViewModelScriptingStruct> USequencerOutlinerScriptingObject::GetDeactivatedNodes()
+{
+	TArray<FSequencerViewModelScriptingStruct> DeactivatedNodes;
+
+	using namespace UE::Sequencer;
+	TViewModelPtr<FOutlinerViewModel> Outliner = WeakOutliner.Pin();
+	if (!Outliner)
+	{
+		FFrame::KismetExecutionMessage(*LOCTEXT("OutlinerInvalid", "Outliner is no longer valid.").ToString(), ELogVerbosity::Error);
+		return DeactivatedNodes;
+	}
+
+	for (const TViewModelPtr<IDeactivatableExtension>& Extension : Outliner->GetRootItem()->GetDescendantsOfType<IDeactivatableExtension>(true))
+	{
+		if (Extension.IsValid() && Extension->IsDeactivated())
+		{
+			DeactivatedNodes.Add(FSequencerViewModelScriptingStruct(Extension.AsModel()));
+		}
+	}
+
+	return DeactivatedNodes;
+}
+
+void USequencerOutlinerScriptingObject::SetDeactivatedNodes(const TArray<FSequencerViewModelScriptingStruct>& InNodes, bool bInDeactivated)
+{
+	using namespace UE::Sequencer;
+
+	for (const FSequencerViewModelScriptingStruct& Item : InNodes)
+	{
+		TViewModelPtr<IDeactivatableExtension> DeactivatableItem = Item.WeakViewModel.ImplicitPin();
+		if (DeactivatableItem)
+		{
+			DeactivatableItem->SetIsDeactivated(bInDeactivated);
+		}
+	}
+}
+
+TArray<FSequencerViewModelScriptingStruct> USequencerOutlinerScriptingObject::GetLockedNodes()
+{
+	TArray<FSequencerViewModelScriptingStruct> LockedNodes;
+
+	using namespace UE::Sequencer;
+	TViewModelPtr<FOutlinerViewModel> Outliner = WeakOutliner.Pin();
+	if (!Outliner)
+	{
+		FFrame::KismetExecutionMessage(*LOCTEXT("OutlinerInvalid", "Outliner is no longer valid.").ToString(), ELogVerbosity::Error);
+		return LockedNodes;
+	}
+
+	for (const TViewModelPtr<ILockableExtension>& Extension : Outliner->GetRootItem()->GetDescendantsOfType<ILockableExtension>(true))
+	{
+		if (Extension.IsValid() && Extension->GetLockState() == ELockableLockState::Locked)
+		{
+			LockedNodes.Add(FSequencerViewModelScriptingStruct(Extension.AsModel()));
+		}
+	}
+
+	return LockedNodes;
+}
+
+void USequencerOutlinerScriptingObject::SetLockedNodes(const TArray<FSequencerViewModelScriptingStruct>& InNodes, bool bInLocked)
+{
+	using namespace UE::Sequencer;
+
+	for (const FSequencerViewModelScriptingStruct& Item : InNodes)
+	{
+		TViewModelPtr<ILockableExtension> LockableItem = Item.WeakViewModel.ImplicitPin();
+		if (LockableItem)
+		{
+			LockableItem->SetIsLocked(bInLocked);
+		}
+	}
+}
+
+TArray<FSequencerViewModelScriptingStruct> USequencerOutlinerScriptingObject::GetPinnedNodes()
+{
+	TArray<FSequencerViewModelScriptingStruct> PinnedNodes;
+
+	using namespace UE::Sequencer;
+	TViewModelPtr<FOutlinerViewModel> Outliner = WeakOutliner.Pin();
+	if (!Outliner)
+	{
+		FFrame::KismetExecutionMessage(*LOCTEXT("OutlinerInvalid", "Outliner is no longer valid.").ToString(), ELogVerbosity::Error);
+		return PinnedNodes;
+	}
+
+	for (const TViewModelPtr<IPinnableExtension>& Extension : Outliner->GetRootItem()->GetDescendantsOfType<IPinnableExtension>(true))
+	{
+		if (Extension.IsValid() && Extension->IsPinned())
+		{
+			PinnedNodes.Add(FSequencerViewModelScriptingStruct(Extension.AsModel()));
+		}
+	}
+
+	return PinnedNodes;
+}
+
+void USequencerOutlinerScriptingObject::SetPinnedNodes(const TArray<FSequencerViewModelScriptingStruct>& InNodes, bool bInPinned)
+{
+	using namespace UE::Sequencer;
+
+	for (const FSequencerViewModelScriptingStruct& Item : InNodes)
+	{
+		TViewModelPtr<IPinnableExtension> PinnableItem = Item.WeakViewModel.ImplicitPin();
+		if (PinnableItem)
+		{
+			PinnableItem->SetPinned(bInPinned);
+		}
+	}
+}
+
 #undef LOCTEXT_NAMESPACE
\ No newline at end of file
diff --git a/Engine/Source/Editor/SequencerCore/Public/Scripting/OutlinerScriptingObject.h b/Engine/Source/Editor/SequencerCore/Public/Scripting/OutlinerScriptingObject.h
--- a/Engine/Source/Editor/SequencerCore/Public/Scripting/OutlinerScriptingObject.h
+++ b/Engine/Source/Editor/SequencerCore/Public/Scripting/OutlinerScriptingObject.h
@@ -36,9 +36,44 @@ public:
 	UFUNCTION(BlueprintCallable, Category="Sequencer Editor")
 	TArray<FSequencerViewModelScriptingStruct> GetSelection() const;
 
-	UFUNCTION(BlueprintCallable, Category="Sequencer Editor")
+	UFUNCTION(BlueprintCallable, Category = "Sequencer Editor")
 	void SetSelection(const TArray<FSequencerViewModelScriptingStruct>& InSelection);
 
+	// Mute
+	UFUNCTION(BlueprintCallable, Category = "Sequencer Editor")
+	TArray<FSequencerViewModelScriptingStruct> GetMuteNodes();
+
+	UFUNCTION(BlueprintCallable, Category = "Sequencer Editor")
+	void SetMuteNodes(const TArray<FSequencerViewModelScriptingStruct>& InNodes, bool bInMuted);
+
+	// Solo
+	UFUNCTION(BlueprintCallable, Category = "Sequencer Editor")
+	TArray<FSequencerViewModelScriptingStruct> GetSoloNodes();
+
+	UFUNCTION(BlueprintCallable, Category = "Sequencer Editor")
+	void SetSoloNodes(const TArray<FSequencerViewModelScriptingStruct>& InNodes, bool bInSoloed);
+
+	// Deactivated
+	UFUNCTION(BlueprintCallable, Category = "Sequencer Editor")
+	TArray<FSequencerViewModelScriptingStruct> GetDeactivatedNodes();
+
+	UFUNCTION(BlueprintCallable, Category = "Sequencer Editor")
+	void SetDeactivatedNodes(const TArray<FSequencerViewModelScriptingStruct>& InNodes, bool bInDeactivated);
+
+	// Locked
+	UFUNCTION(BlueprintCallable, Category = "Sequencer Editor")
+	TArray<FSequencerViewModelScriptingStruct> GetLockedNodes();
+
+	UFUNCTION(BlueprintCallable, Category = "Sequencer Editor")
+	void SetLockedNodes(const TArray<FSequencerViewModelScriptingStruct>& InNodes, bool bInLocked);
+
+	// Pinned
+	UFUNCTION(BlueprintCallable, Category = "Sequencer Editor")
+	TArray<FSequencerViewModelScriptingStruct> GetPinnedNodes();
+
+	UFUNCTION(BlueprintCallable, Category = "Sequencer Editor")
+	void SetPinnedNodes(const TArray<FSequencerViewModelScriptingStruct>& InNodes, bool bInPinned);
+
 protected:
 
 	UE::Sequencer::TWeakViewModelPtr<UE::Sequencer::FOutlinerViewModel> WeakOutliner;
diff --git a/Engine/Source/Runtime/MovieScene/Private/MovieScene.cpp b/Engine/Source/Runtime/MovieScene/Private/MovieScene.cpp
--- a/Engine/Source/Runtime/MovieScene/Private/MovieScene.cpp
+++ b/Engine/Source/Runtime/MovieScene/Private/MovieScene.cpp
@@ -570,11 +570,15 @@ void UMovieScene::SetTickResolutionDirectly(FFrameRate InTickResolution)
 
 void UMovieScene::AddNewBindingTag(const FName& NewTag)
 {
+	Modify();
+
 	BindingGroups.Add(NewTag);
 }
 
 void UMovieScene::TagBinding(const FName& NewTag, const UE::MovieScene::FFixedObjectBindingID& BindingToTag)
 {
+	Modify();
+
 	FMovieSceneObjectBindingID SerializedID = BindingToTag;
 	BindingGroups.FindOrAdd(NewTag).IDs.AddUnique(SerializedID);
 }
@@ -584,6 +588,8 @@ void UMovieScene::UntagBinding(const FName& Tag, const UE::MovieScene::FFixedObj
 	FMovieSceneObjectBindingID PredicateID = Binding;
 	if (FMovieSceneObjectBindingIDs* Array = BindingGroups.Find(Tag))
 	{
+		Modify();
+
 		Array->IDs.Remove(PredicateID);
 		if (Array->IDs.Num() == 0)
 		{
@@ -594,6 +600,8 @@ void UMovieScene::UntagBinding(const FName& Tag, const UE::MovieScene::FFixedObj
 
 void UMovieScene::RemoveTag(const FName& TagToRemove)
 {
+	Modify();
+
 	BindingGroups.Remove(TagToRemove);
 }
 
