Sequencer: When considering whether children should be expanded or not, only consider nodes if they have children

This fixes a bug when toggling expansion states where it thinks it doesn't need to collapse because nodes without children report back as being collapsed

diff --git a/Engine/Source/Editor/SequencerCore/Private/MVVM/Views/SOutlinerView.cpp b/Engine/Source/Editor/SequencerCore/Private/MVVM/Views/SOutlinerView.cpp
--- a/Engine/Source/Editor/SequencerCore/Private/MVVM/Views/SOutlinerView.cpp
+++ b/Engine/Source/Editor/SequencerCore/Private/MVVM/Views/SOutlinerView.cpp
@@ -1037,9 +1037,13 @@ bool ShouldExpand(const TArrayView<const TWeakViewModelPtr<IOutlinerExtension>>&
 		{
 			for (const TViewModelPtr<IOutlinerExtension>& Child : SelectedItem.AsModel()->GetChildrenOfType<IOutlinerExtension>())
 			{
-				if (!Child->IsExpanded())
+				FViewModelListIterator ChildChildren = Child.AsModel()->GetChildren(EViewModelListType::Outliner);
+				if (ChildChildren)
 				{
-					return true;
+					if (!Child->IsExpanded())
+					{
+						return true;
+					}
 				}
 			}
 		}
