Sequencer: Fix selection desync with nested object bindings

Object binding nodes are considered indirectly selected if the track area's ancestor object binding node is selected. However, object binding nodes can be children of other object binding nodes, so the comparison must be done with all the ancestor object binding nodes.

This fixes an issue where if you have this hierarchy:
- Cube
- Static Mesh Component
-Transform Track
with the Static Mesh Component selected, selecting a layer bar will deselect the Static Mesh Component and the Cube.

Note, this does not force the cube to be selected when selecting a layer bar, which has been the case in the past with UE-227929

diff --git a/Engine/Source/Editor/Sequencer/Private/MVVM/Selection.cpp b/Engine/Source/Editor/Sequencer/Private/MVVM/Selection.cpp
--- a/Engine/Source/Editor/Sequencer/Private/MVVM/Selection.cpp
+++ b/Engine/Source/Editor/Sequencer/Private/MVVM/Selection.cpp
@@ -212,7 +212,8 @@ void FSequencerSelection::PreBroadcastChangeEvent()
 		bool bFound = false;
 		bool bAnyIndirectSelection = false;
 
-		TSharedPtr<IObjectBindingExtension> ObjectBindingItem = OutlinerItem.AsModel()->FindAncestorOfType<IObjectBindingExtension>(true);
+		const TSet<TViewModelPtr<IObjectBindingExtension>> ObjectBindingAncestors = OutlinerItem.AsModel()->GetAncestorsOfType<IObjectBindingExtension>(true).Populate<TSet>();
+
 
 		for (TViewModelPtr<IOutlinerExtension> IndirectItem : IterateIndirectOutlinerSelection())
 		{
@@ -224,9 +225,9 @@ void FSequencerSelection::PreBroadcastChangeEvent()
 				break;
 			}
 
-			for (TSharedPtr<IObjectBindingExtension> IndirectObjectBindingItem : IndirectItem.AsModel()->GetAncestorsOfType<IObjectBindingExtension>(true))
+			for (TViewModelPtr<IObjectBindingExtension> IndirectObjectBindingItem : IndirectItem.AsModel()->GetAncestorsOfType<IObjectBindingExtension>(true))
 			{
-				if (IndirectObjectBindingItem == ObjectBindingItem)
+				if (ObjectBindingAncestors.Contains(IndirectObjectBindingItem))
 				{
 					bFound = true;
 					break;
diff --git a/Engine/Source/Editor/SequencerCore/Public/MVVM/ViewModels/ViewModelIterators.h b/Engine/Source/Editor/SequencerCore/Public/MVVM/ViewModels/ViewModelIterators.h
--- a/Engine/Source/Editor/SequencerCore/Public/MVVM/ViewModels/ViewModelIterators.h
+++ b/Engine/Source/Editor/SequencerCore/Public/MVVM/ViewModels/ViewModelIterators.h
@@ -48,10 +48,88 @@ struct FViewModelIterationState
 // WARNING: most of the class and methods in this header are not safe to use from a view-model's
 // constructor because they rely on shared-pointers to be available for that view-model.
 
+template<typename DerivedType, typename EnumeratedType>
+struct TPopulatableIteratorMixin
+{
+
+	/**
+	 * Populate the specified container with the remaining items in this iterator.
+	 * ContainerType must have a .Add method that accepts the enumerated type.
+	 * The state of this iterator remains unchanged.
+	 * 
+	 * Example usage:
+	 * 
+	 * TSet<FViewModelPtr> Set;
+	 * It.Populate(Set);
+	 */
+	template<typename ContainerType>
+	void Populate(ContainerType& OutContainer) const
+	{
+		for (const EnumeratedType& Item : static_cast<const DerivedType&>(*this))
+		{
+			OutContainer.Add(Item);
+		}
+	}
+
+	/**
+	 * Create a new container of the templated type and populate it with the remaining items in this iterator.
+	 * ContainerType must have a .Add method that accepts the enumerated type.
+	 * The state of this iterator remains unchanged.
+	 * 
+	 * Example usage:
+	 * 
+	 * TSet<TViewModelPtr<FTrackModel>> Tracks = ViewModel->GetDescendentsOfType<FTrackModel>().Populate<TSet>();
+	 */
+	template<template<typename T, typename ...Rest> class ContainerType>
+	ContainerType<EnumeratedType> Populate() const
+	{
+		ContainerType<EnumeratedType> Temp;
+		Populate(Temp);
+		return Temp;
+	}
+
+	/**
+	 * Populate the specified array with the remaining items in this iterator
+	 * The state of this iterator remains unchanged
+	 */
+	void ToArray(TArray<EnumeratedType>& OutArray) const
+	{
+		Populate(OutArray);
+	}
+
+	/**
+	 * Return the remaining items in this iterator as an array
+	 * The state of this iterator remains unchanged
+	 */
+	TArray<EnumeratedType> ToArray() const
+	{
+		return Populate<TArray>();
+	}
+
+	/**
+	 * Linearly search through this iterator to find a predicate using a projection
+	 * For example:
+	 *      UMovieSceneTrack* TrackToFind = ...;
+	 *      Model->GetChildrenOfType<FTrackModel>().FindBy(TrackToFind, &FTrackModel::GetTrack)
+	 */
+	template<typename ValueType, typename ProjectionType>
+	EnumeratedType FindBy(const ValueType& Value, ProjectionType Proj) const
+	{
+		for (const EnumeratedType& Item : static_cast<const DerivedType&>(*this))
+		{
+			if (Invoke(Proj, *Item) == Value)
+			{
+				return Item;
+			}
+		}
+		return nullptr;
+	}
+};
+
 /**
  * Iterates all models that exist in a single list from the start model
  */
-struct FViewModelSubListIterator
+struct FViewModelSubListIterator : TPopulatableIteratorMixin<FViewModelSubListIterator, FViewModelPtr>
 {
 	FViewModelSubListIterator() = default;
 
@@ -83,29 +161,6 @@ struct FViewModelSubListIterator
 	FORCEINLINE FViewModelSubListIterator begin() const { return *this; }
 	FORCEINLINE FViewModelSubListIterator end() const   { return FViewModelSubListIterator{ nullptr }; }
 
-	/**
-	 * Populate the specified array with the remaining items in this iterator
-	 * The state of this iterator remains unchanged
-	 */
-	void ToArray(TArray<FViewModelPtr>& OutArray) const
-	{
-		for (FViewModelSubListIterator Temp = *this; Temp; ++Temp)
-		{
-			OutArray.Add(*Temp);
-		}
-	}
-
-	/**
-	 * Return the remaining items in this iterator as an array
-	 * The state of this iterator remains unchanged
-	 */
-	TArray<FViewModelPtr> ToArray() const
-	{
-		TArray<FViewModelPtr> Temp;
-		ToArray(Temp);
-		return Temp;
-	}
-
 private:
 
 	/** The pointed-to data model */
@@ -115,7 +170,7 @@ private:
 /**
  * Iterates all models contained in the specified list chain
  */
-struct FViewModelListIterator
+struct FViewModelListIterator : TPopulatableIteratorMixin<FViewModelListIterator, FViewModelPtr>
 {
 	FViewModelListIterator(EViewModelListType InFilter = EViewModelListType::Everything)
 		: NextList(nullptr)
@@ -154,49 +209,6 @@ struct FViewModelListIterator
 	FORCEINLINE FViewModelListIterator begin() const { return *this; }
 	FORCEINLINE FViewModelListIterator end() const   { return FViewModelListIterator{ nullptr }; }
 
-	/**
-	 * Populate the specified array with the remaining items in this iterator
-	 * The state of this iterator remains unchanged
-	 */
-	void ToArray(TArray<FViewModelPtr>& OutArray) const
-	{
-		for (FViewModelListIterator Temp = *this; Temp; ++Temp)
-		{
-			OutArray.Add(*Temp);
-		}
-	}
-
-	/**
-	 * Return the remaining items in this iterator as an array
-	 * The state of this iterator remains unchanged
-	 */
-	TArray<FViewModelPtr> ToArray() const
-	{
-		TArray<FViewModelPtr> Temp;
-		ToArray(Temp);
-		return Temp;
-	}
-
-	/**
-	 * Linearly search through this iterator to find a predicate using a projection
-	 * For example, to find the first item that has children of its own:
-	 *      const bool bHasChildren = true;
-	 *      Model->GetChildren().FindBy(bHasChildren, &FViewModel::HasChildren)
-	 */
-	template<typename ValueType, typename ProjectionType>
-	FViewModelPtr FindBy(const ValueType& Value, ProjectionType Proj) const
-	{
-		for (FViewModelListIterator Temp = *this; Temp; ++Temp)
-		{
-			FViewModelPtr Item = *Temp;
-			if (Invoke(Proj, *Item) == Value)
-			{
-				return Item;
-			}
-		}
-		return nullptr;
-	}
-
 private:
 
 	void FindNextValidIterator()
@@ -484,7 +496,7 @@ protected:
  * so that it can be used as a typed iterator.
  */
 template<typename T, typename IteratorType>
-struct TTypedIteratorBase : public IteratorType
+struct TTypedIteratorBase : public IteratorType, public TPopulatableIteratorMixin<TTypedIteratorBase<T, IteratorType>, TViewModelPtr<T>>
 {
 	template<typename ...ArgTypes>
 	TTypedIteratorBase(ArgTypes&&... InArgs)
@@ -502,52 +514,13 @@ struct TTypedIteratorBase : public IteratorType
 		GetCastResultOrContinue();
 	}
 
-	/**
-	 * Populate the specified array with the remaining items in this iterator
-	 * The state of this iterator remains unchanged
-	 */
-	void ToArray(TArray<TViewModelPtr<T>>& OutArray) const
-	{
-		for (TTypedIteratorBase<T, IteratorType> Temp = *this; Temp; ++Temp)
-		{
-			OutArray.Add(*Temp);
-		}
-	}
-
-	/**
-	 * Return the remaining items in this iterator as an array
-	 * The state of this iterator remains unchanged
-	 */
-	TArray<TViewModelPtr<T>> ToArray() const
-	{
-		TArray<TViewModelPtr<T>> Temp;
-		ToArray(Temp);
-		return Temp;
-	}
-
-	/**
-	 * Linearly search through this iterator to find a predicate using a projection
-	 * For example:
-	 *      UMovieSceneTrack* TrackToFind = ...;
-	 *      Model->GetChildrenOfType<FTrackModel>().FindBy(TrackToFind, &FTrackModel::GetTrack)
-	 */
-	template<typename ValueType, typename ProjectionType>
-	TViewModelPtr<T> FindBy(const ValueType& Value, ProjectionType Proj) const
-	{
-		for (TTypedIteratorBase<T, IteratorType> Temp = *this; Temp; ++Temp)
-		{
-			TViewModelPtr<T> Item = *Temp;
-			if (Invoke(Proj, *Item) == Value)
-			{
-				return Item;
-			}
-		}
-		return nullptr;
-	}
-
 	FORCEINLINE TTypedIterator<T, IteratorType> begin() const { return IteratorType::begin(); }
 	FORCEINLINE TTypedIterator<T, IteratorType> end() const   { return IteratorType::end(); }
 
+	using TPopulatableIteratorMixin<TTypedIteratorBase<T, IteratorType>, TViewModelPtr<T>>::ToArray;
+	using TPopulatableIteratorMixin<TTypedIteratorBase<T, IteratorType>, TViewModelPtr<T>>::Populate;
+	using TPopulatableIteratorMixin<TTypedIteratorBase<T, IteratorType>, TViewModelPtr<T>>::FindBy;
+
 protected:
 
 	void GetCastResultOrContinue()
