[Core] Cleanup bad usage of begin/end on tset

#tests Compiled mostly

diff --git a/Engine/Plugins/Animation/ControlRig/Source/ControlRig/Private/Sequencer/MovieSceneControlRigParameterSection.cpp b/Engine/Plugins/Animation/ControlRig/Source/ControlRig/Private/Sequencer/MovieSceneControlRigParameterSection.cpp
--- a/Engine/Plugins/Animation/ControlRig/Source/ControlRig/Private/Sequencer/MovieSceneControlRigParameterSection.cpp
+++ b/Engine/Plugins/Animation/ControlRig/Source/ControlRig/Private/Sequencer/MovieSceneControlRigParameterSection.cpp
@@ -1805,13 +1805,12 @@ const FName& UMovieSceneControlRigParameterSection::FindControlNameFromConstrain
 		const int32 Index = Entry - ConstraintsChannels.GetData();
 
 		// look for info referencing that constraint index
-		using NameInfoIterator = TMap<FName, FChannelMapInfo>::TRangedForConstIterator;
-		for (NameInfoIterator It = ControlChannelMap.begin(); It; ++It)
+		for (const TPair<FName, FChannelMapInfo>& It : ControlChannelMap)
 		{
-			const FChannelMapInfo& Info = It->Value;
+			const FChannelMapInfo& Info = It.Value;
 			if (Info.ConstraintsIndex.Contains(Index))
 			{
-				return It->Key;
+				return It.Key;
 			}
 		}
 	}
diff --git a/Engine/Plugins/Animation/ControlRig/Source/ControlRigEditor/Private/Sequencer/ControlRigParameterTrackEditor.cpp b/Engine/Plugins/Animation/ControlRig/Source/ControlRigEditor/Private/Sequencer/ControlRigParameterTrackEditor.cpp
--- a/Engine/Plugins/Animation/ControlRig/Source/ControlRigEditor/Private/Sequencer/ControlRigParameterTrackEditor.cpp
+++ b/Engine/Plugins/Animation/ControlRig/Source/ControlRigEditor/Private/Sequencer/ControlRigParameterTrackEditor.cpp
@@ -3661,13 +3661,12 @@ namespace
 			// get constraint name
 			auto GetControlName = [InSection, Index]()
 			{
-				using NameInfoIterator = TMap<FName, FChannelMapInfo>::TRangedForConstIterator;
-				for (NameInfoIterator It = InSection->ControlChannelMap.begin(); It; ++It)
+				for (const TPair<FName, FChannelMapInfo>& It : InSection->ControlChannelMap)
 				{
-					const FChannelMapInfo& Info = It->Value;
+					const FChannelMapInfo& Info = It.Value;
 					if (Info.ConstraintsIndex.Contains(Index))
 					{
-						return It->Key;
+						return It.Key;
 					}
 				}
 
diff --git a/Engine/Plugins/EnhancedInput/Source/EnhancedInput/Private/UserSettings/EnhancedInputUserSettings.cpp b/Engine/Plugins/EnhancedInput/Source/EnhancedInput/Private/UserSettings/EnhancedInputUserSettings.cpp
--- a/Engine/Plugins/EnhancedInput/Source/EnhancedInput/Private/UserSettings/EnhancedInputUserSettings.cpp
+++ b/Engine/Plugins/EnhancedInput/Source/EnhancedInput/Private/UserSettings/EnhancedInputUserSettings.cpp
@@ -1166,7 +1166,7 @@ void UEnhancedInputUserSettings::MapPlayerKey(const FMapPlayerKeyArgs& InArgs, F
 		const int32 NumMappings = MappingRow->Mappings.Num();
 		if (InArgs.bCreateMatchingSlotIfNeeded && NumMappings > 0)
 		{
-			const auto ExistingMapping = MappingRow->Mappings.begin();
+			const FPlayerKeyMapping* ExistingMapping = MappingRow->Mappings.FindArbitraryElement();
 			
 			// Add a default mapping to this row
 			FPlayerKeyMapping PlayerMappingData = {};
diff --git a/Engine/Plugins/Experimental/Mover/Source/Mover/Private/Debug/GameplayDebuggerCategory_Mover.cpp b/Engine/Plugins/Experimental/Mover/Source/Mover/Private/Debug/GameplayDebuggerCategory_Mover.cpp
--- a/Engine/Plugins/Experimental/Mover/Source/Mover/Private/Debug/GameplayDebuggerCategory_Mover.cpp
+++ b/Engine/Plugins/Experimental/Mover/Source/Mover/Private/Debug/GameplayDebuggerCategory_Mover.cpp
@@ -125,16 +125,16 @@ void FGameplayDebuggerCategory_Mover::CollectData(APlayerController* OwnerPC, AA
 		DataPack.MoveIntent = MyMoverComponent->GetMovementIntent();
 		DataPack.Velocity = MyMoverComponent->GetVelocity();
 
-		for (auto ModeIter = MyMoverComponent->MovementModes.begin(); ModeIter; ++ModeIter)
+		for (const TPair<FName, TObjectPtr<UBaseMovementMode>>& ModeIter : MyMoverComponent->MovementModes)
 		{
-			UBaseMovementMode* MappedMode = ModeIter->Value;
-			DataPack.ModeMap.Add(FString::Printf(TEXT("%s => %s"), *ModeIter->Key.ToString(), (MappedMode ? *MappedMode->GetClass()->GetName() : TEXT("null"))));
+			UBaseMovementMode* MappedMode = ModeIter.Value;
+			DataPack.ModeMap.Add(FString::Printf(TEXT("%s => %s"), *ModeIter.Key.ToString(), (MappedMode ? *MappedMode->GetClass()->GetName() : TEXT("null"))));
 
-			if (ModeIter->Key == MyMoverComponent->GetMovementModeName())
+			if (ModeIter.Key == MyMoverComponent->GetMovementModeName())
 			{
-				for (UBaseMovementModeTransition* Transition : MyMoverComponent->MovementModes[ModeIter->Key]->Transitions)
+				for (UBaseMovementModeTransition* Transition : MyMoverComponent->MovementModes[ModeIter.Key]->Transitions)
 				{
-					DataPack.ActiveTransitions.Add(FString::Printf(TEXT("%s (%s)"), (Transition ? *Transition->GetClass()->GetName() : TEXT("null")), *ModeIter->Key.ToString()));
+					DataPack.ActiveTransitions.Add(FString::Printf(TEXT("%s (%s)"), (Transition ? *Transition->GetClass()->GetName() : TEXT("null")), *ModeIter.Key.ToString()));
 				}
 			}
 		}
diff --git a/Engine/Plugins/Interchange/Runtime/Source/Import/Private/Gltf/InterchangeGLTFMaterial.cpp b/Engine/Plugins/Interchange/Runtime/Source/Import/Private/Gltf/InterchangeGLTFMaterial.cpp
--- a/Engine/Plugins/Interchange/Runtime/Source/Import/Private/Gltf/InterchangeGLTFMaterial.cpp
+++ b/Engine/Plugins/Interchange/Runtime/Source/Import/Private/Gltf/InterchangeGLTFMaterial.cpp
@@ -85,7 +85,7 @@ namespace UE::Interchange::GLTFMaterials
 		}
 
 		ensure(false);
-		return ShadingModelToMaterialInformation.begin()->Value.MaterialFunctionPath;
+		return ShadingModelToMaterialInformation.FindArbitraryElement()->Value.MaterialFunctionPath;
 	}
 
 	TArray<FString> GetOutputs(EShadingModel ShadingModel)
@@ -96,7 +96,7 @@ namespace UE::Interchange::GLTFMaterials
 		}
 
 		ensure(false);
-		return ShadingModelToMaterialInformation.begin()->Value.MaterialFunctionOutputs;
+		return ShadingModelToMaterialInformation.FindArbitraryElement()->Value.MaterialFunctionOutputs;
 	}
 
 	//Precedence order based on the logic in the InterchangeGenericMaterialPipeline: ClearCoat > Sheen > unlit
diff --git a/Engine/Plugins/Interchange/Runtime/Source/Parsers/GLTFCore/Private/GLTF/GLTFReader.cpp b/Engine/Plugins/Interchange/Runtime/Source/Parsers/GLTFCore/Private/GLTF/GLTFReader.cpp
--- a/Engine/Plugins/Interchange/Runtime/Source/Parsers/GLTFCore/Private/GLTF/GLTFReader.cpp
+++ b/Engine/Plugins/Interchange/Runtime/Source/Parsers/GLTFCore/Private/GLTF/GLTFReader.cpp
@@ -1401,7 +1401,7 @@ namespace GLTF
 						{
 							//Scenario is that the a Skin is instantiated at the end of another skin
 							//(Prime example is the RecursiveSkeleton gltf sample file.)
-							ParentGlobalInverseBindMatrix = ParentNode.SkinIndexToGlobalInverseBindMatrix.begin().Value();
+							ParentGlobalInverseBindMatrix = ParentNode.SkinIndexToGlobalInverseBindMatrix.FindArbitraryElement()->Value;
 						}
 
 						FTransform LocalBindTransform = FTransform::Identity;
@@ -1451,7 +1451,7 @@ namespace GLTF
 
 		for (const FNode& CurrentNode : Asset->Nodes)
 		{
-			TMap<int, FTransform>::TRangedForConstIterator Iter(CurrentNode.SkinIndexToLocalBindTransform.begin());
+			TMap<int, FTransform>::TConstIterator Iter(CurrentNode.SkinIndexToLocalBindTransform.CreateConstIterator());
 			FTransform ToCompareAgainst = Iter ? Iter.Value() : FTransform();
 			++Iter;
 			for (; Iter; ++Iter)
diff --git a/Engine/Plugins/Interchange/Runtime/Source/Pipelines/Private/InterchangePipelineMeshesUtilities.cpp b/Engine/Plugins/Interchange/Runtime/Source/Pipelines/Private/InterchangePipelineMeshesUtilities.cpp
--- a/Engine/Plugins/Interchange/Runtime/Source/Pipelines/Private/InterchangePipelineMeshesUtilities.cpp
+++ b/Engine/Plugins/Interchange/Runtime/Source/Pipelines/Private/InterchangePipelineMeshesUtilities.cpp
@@ -478,7 +478,7 @@ UInterchangePipelineMeshesUtilities* UInterchangePipelineMeshesUtilities::Create
 					}
 				}
 			);
-			PipelineMeshesUtilities->MeshGeometriesPerMeshUid.begin().Value().AttachedSocketUids = SocketUIDs;
+			PipelineMeshesUtilities->MeshGeometriesPerMeshUid.FindArbitraryElement()->Value.AttachedSocketUids = SocketUIDs;
 		}
 		else
 		{
diff --git a/Engine/Plugins/Mutable/Source/CustomizableObjectEditor/Private/MuCOE/SCustomizableObjectEditorViewport.cpp b/Engine/Plugins/Mutable/Source/CustomizableObjectEditor/Private/MuCOE/SCustomizableObjectEditorViewport.cpp
--- a/Engine/Plugins/Mutable/Source/CustomizableObjectEditor/Private/MuCOE/SCustomizableObjectEditorViewport.cpp
+++ b/Engine/Plugins/Mutable/Source/CustomizableObjectEditor/Private/MuCOE/SCustomizableObjectEditorViewport.cpp
@@ -700,7 +700,7 @@ void SCustomizableObjectEditorViewportTabBody::OnTakeHighResScreenshot()
 	TMap<FName, TWeakObjectPtr<UDebugSkelMeshComponent>>& Comps = LevelViewportClient->GetPreviewMeshComponents();
 	if (!Comps.IsEmpty())
 	{
-		CustomizableObjectHighresScreenshot = SCustomizableObjectHighresScreenshot::OpenDialog(ViewportWidget->GetSceneViewport(), LevelViewportClient, Comps.begin().Value().Get(), PreviewScenePtr);
+		CustomizableObjectHighresScreenshot = SCustomizableObjectHighresScreenshot::OpenDialog(ViewportWidget->GetSceneViewport(), LevelViewportClient, Comps.FindArbitraryElement()->Value.Get(), PreviewScenePtr);
 	}
 }
 
diff --git a/Engine/Plugins/Mutable/Source/MutableRuntime/Private/MuR/Parameters.cpp b/Engine/Plugins/Mutable/Source/MutableRuntime/Private/MuR/Parameters.cpp
--- a/Engine/Plugins/Mutable/Source/MutableRuntime/Private/MuR/Parameters.cpp
+++ b/Engine/Plugins/Mutable/Source/MutableRuntime/Private/MuR/Parameters.cpp
@@ -158,7 +158,7 @@ namespace mu
             return nullptr;
         }
 
-		TMap< TArray<int32>, FParameterValue >::TRangedForIterator it = m_pD->MultiValues[ParamIndex].begin();
+		TMap< TArray<int32>, FParameterValue >::TIterator it = m_pD->MultiValues[ParamIndex].CreateIterator();
         for ( int32 i=0; i< ValueIndex; ++i )
         {
             ++it;
diff --git a/Engine/Plugins/PCG/Source/PCG/Private/SpatialAlgo/PCGMarchingSquares.cpp b/Engine/Plugins/PCG/Source/PCG/Private/SpatialAlgo/PCGMarchingSquares.cpp
--- a/Engine/Plugins/PCG/Source/PCG/Private/SpatialAlgo/PCGMarchingSquares.cpp
+++ b/Engine/Plugins/PCG/Source/PCG/Private/SpatialAlgo/PCGMarchingSquares.cpp
@@ -407,8 +407,8 @@ TArray<PCGSpatialAlgo::FPCGMarchingSquareResult> PCGSpatialAlgo::MarchingSquares
 	TArray<FPCGMarchingSquareResult> FinalResult;
 	while (!Candidates.IsEmpty())
 	{
-		int32 StartingIndex = Candidates.begin()->Key;
-		Edge StartingEdgeMask = Candidates.begin()->Value;
+		int32 StartingIndex = Candidates.FindArbitraryElement()->Key;
+		Edge StartingEdgeMask = Candidates.FindArbitraryElement()->Value;
 		check(!!StartingEdgeMask);
 
 		TArray<FVector2D> Head;
diff --git a/Engine/Plugins/Runtime/RigVM/Source/RigVM/Private/RigVMCore/RigVMStruct.cpp b/Engine/Plugins/Runtime/RigVM/Source/RigVM/Private/RigVMCore/RigVMStruct.cpp
--- a/Engine/Plugins/Runtime/RigVM/Source/RigVM/Private/RigVMCore/RigVMStruct.cpp
+++ b/Engine/Plugins/Runtime/RigVM/Source/RigVM/Private/RigVMCore/RigVMStruct.cpp
@@ -115,7 +115,7 @@ FRigVMStructUpgradeInfo FRigVMStructUpgradeInfo::MakeFromStructToFactory(UScript
 		if(Types.Num() == 1 && Infos.Num() == 1)
 		{
 			FRigVMTemplateTypeMap NewTypes;
-			NewTypes.Add(Infos[0].Name, Types.begin().Value());
+			NewTypes.Add(Infos[0].Name, Types.FindArbitraryElement()->Value);
 			Swap(NewTypes, Types);
 		}
 
diff --git a/Engine/Plugins/VirtualProduction/LevelSnapshots/Source/LevelSnapshots/Private/Tests/SubobjectTests.cpp b/Engine/Plugins/VirtualProduction/LevelSnapshots/Source/LevelSnapshots/Private/Tests/SubobjectTests.cpp
--- a/Engine/Plugins/VirtualProduction/LevelSnapshots/Source/LevelSnapshots/Private/Tests/SubobjectTests.cpp
+++ b/Engine/Plugins/VirtualProduction/LevelSnapshots/Source/LevelSnapshots/Private/Tests/SubobjectTests.cpp
@@ -818,10 +818,10 @@ namespace UE::LevelSnapshots::Private::Tests
 			
 				ASnapshotTestActor* ActorB_Snapshot = Cast<ASnapshotTestActor>(OptionalActorB_Snapshot->Get());
 				TestTrue(TEXT("ObjectArray contains null"), ActorB_Snapshot->ObjectArray.Num() == 1 && ActorB_Snapshot->ObjectArray[0] == nullptr);
-				TestTrue(TEXT("ObjectSet contains null"), ActorB_Snapshot->ObjectSet.Num() == 1 && ActorB_Snapshot->ObjectSet.begin().ElementIt->Value == nullptr);
+				TestTrue(TEXT("ObjectSet contains null"), ActorB_Snapshot->ObjectSet.Num() == 1 && *ActorB_Snapshot->ObjectSet.FindArbitraryElement() == nullptr);
 				TestTrue(TEXT("ObjectMap contains null"), ActorB_Snapshot->ObjectMap.Num() == 1 && ActorB_Snapshot->ObjectMap.Find("MapKey") && ActorB_Snapshot->ObjectMap["MapKey"] == nullptr);
 				TestTrue(TEXT("WeakObjectPtrArray contains null"), ActorB_Snapshot->WeakObjectPtrArray.Num() == 1 && ActorB_Snapshot->WeakObjectPtrArray[0] == nullptr);
-				TestTrue(TEXT("WeakObjectPtrSet contains null"), ActorB_Snapshot->WeakObjectPtrSet.Num() == 1 && ActorB_Snapshot->WeakObjectPtrSet.begin().ElementIt->Value == nullptr);
+				TestTrue(TEXT("WeakObjectPtrSet contains null"), ActorB_Snapshot->WeakObjectPtrSet.Num() == 1 && *ActorB_Snapshot->WeakObjectPtrSet.FindArbitraryElement() == nullptr);
 				TestTrue(TEXT("WeakObjectPtrMap contains null"), ActorB_Snapshot->WeakObjectPtrMap.Num() == 1 && ActorB_Snapshot->WeakObjectPtrMap.Find("MapKey") && ActorB_Snapshot->WeakObjectPtrMap["MapKey"] == nullptr);
 			})
 
diff --git a/Engine/Source/Editor/Sequencer/Public/MVVM/Selection/Selection.h b/Engine/Source/Editor/Sequencer/Public/MVVM/Selection/Selection.h
--- a/Engine/Source/Editor/Sequencer/Public/MVVM/Selection/Selection.h
+++ b/Engine/Source/Editor/Sequencer/Public/MVVM/Selection/Selection.h
@@ -209,8 +209,8 @@ struct FIndirectOutlinerSelectionIterator
 		}
 	}
 
-	FORCEINLINE TSelectionSetIteratorState<TWeakViewModelPtr<IOutlinerExtension>> begin() const { return TSelectionSetIteratorState<TWeakViewModelPtr<IOutlinerExtension>>(SelectionSet->begin()); }
-	FORCEINLINE TSelectionSetIteratorState<TWeakViewModelPtr<IOutlinerExtension>> end() const   { return TSelectionSetIteratorState<TWeakViewModelPtr<IOutlinerExtension>>(SelectionSet->end()); }
+	FORCEINLINE TSelectionSetIteratorState<TWeakViewModelPtr<IOutlinerExtension>> begin() const { return TSelectionSetIteratorState<TWeakViewModelPtr<IOutlinerExtension>>(SelectionSet->begin(), SelectionSet->end()); }
+	FORCEINLINE TSelectionSetIteratorState<TWeakViewModelPtr<IOutlinerExtension>> end() const   { return TSelectionSetIteratorState<TWeakViewModelPtr<IOutlinerExtension>>(SelectionSet->end(), SelectionSet->end()); }
 };
 
 } // namespace UE::Sequencer
\ No newline at end of file
diff --git a/Engine/Source/Editor/SequencerCore/Public/MVVM/Selection/SequencerCoreSelectionIterators.h b/Engine/Source/Editor/SequencerCore/Public/MVVM/Selection/SequencerCoreSelectionIterators.h
--- a/Engine/Source/Editor/SequencerCore/Public/MVVM/Selection/SequencerCoreSelectionIterators.h
+++ b/Engine/Source/Editor/SequencerCore/Public/MVVM/Selection/SequencerCoreSelectionIterators.h
@@ -32,13 +32,13 @@ struct TSelectionSetIteratorImpl
 	 */
 	const T& operator*() const
 	{
-		return *Iterator;
+		return *IteratorBegin;
 	}
 
 protected:
 
-	explicit TSelectionSetIteratorImpl(typename TSet<T>::TRangedForConstIterator InIterator)
-		: Iterator(InIterator)
+	explicit TSelectionSetIteratorImpl(typename TSet<T>::TRangedForConstIterator InIteratorBegin, typename TSet<T>::TRangedForConstIterator InIteratorEnd)
+	: IteratorBegin(InIteratorBegin), IteratorEnd(InIteratorEnd)
 	{}
 
 	/**
@@ -47,7 +47,8 @@ protected:
 	constexpr void SkipInvalid()
 	{}
 
-	typename TSet<T>::TRangedForConstIterator Iterator;
+	typename TSet<T>::TRangedForConstIterator IteratorBegin;
+	typename TSet<T>::TRangedForConstIterator IteratorEnd;
 };
 
 /**
@@ -56,26 +57,21 @@ protected:
 template<typename T, typename ImplType = TSelectionSetIteratorImpl<T>>
 struct TSelectionSetIteratorState : ImplType
 {
-	explicit TSelectionSetIteratorState(typename TSet<T>::TRangedForConstIterator InIterator)
-		: ImplType(InIterator)
+	explicit TSelectionSetIteratorState(typename TSet<T>::TRangedForConstIterator InIteratorBegin, typename TSet<T>::TRangedForConstIterator InIteratorEnd)
+		: ImplType(InIteratorBegin, InIteratorEnd)
 	{
 		this->SkipInvalid();
 	}
 
 	void operator++()
 	{
-		++this->Iterator;
+		++this->IteratorBegin;
 		this->SkipInvalid();
 	}
 
-	explicit operator bool() const
-	{
-		return static_cast<bool>(this->Iterator);
-	}
-
 	friend bool operator!=(const TSelectionSetIteratorState<T, ImplType>& A, const TSelectionSetIteratorState<T, ImplType>& B)
 	{
-		return A.Iterator != B.Iterator;
+		return A.IteratorBegin != B.IteratorBegin && ensure(A.IteratorEnd == B.IteratorEnd);
 	}
 };
 
@@ -93,25 +89,26 @@ struct TFilteredSelectionSetIteratorImpl
 
 protected:
 
-	explicit TFilteredSelectionSetIteratorImpl(typename TSet<KeyType>::TRangedForConstIterator InIterator)
-		: Iterator(InIterator)
+	explicit TFilteredSelectionSetIteratorImpl(typename TSet<KeyType>::TRangedForConstIterator InIteratorBegin, typename TSet<KeyType>::TRangedForConstIterator InIteratorEnd)
+		: IteratorBegin(InIteratorBegin), IteratorEnd(InIteratorEnd)
 	{}
 
 	void SkipInvalid()
 	{
 		CurrentValue = nullptr;
-		while (Iterator)
+		while (IteratorBegin != IteratorEnd)
 		{
-			CurrentValue = CastViewModel<FilterType>(Iterator->Pin());
+			CurrentValue = CastViewModel<FilterType>(IteratorBegin->Pin());
 			if (CurrentValue)
 			{
 				break;
 			}
-			++Iterator;
+			++IteratorBegin;
 		}
 	}
 
-	typename TSet<KeyType>::TRangedForConstIterator Iterator;
+	typename TSet<KeyType>::TRangedForConstIterator IteratorBegin;
+	typename TSet<KeyType>::TRangedForConstIterator IteratorEnd;
 	TViewModelPtr<FilterType> CurrentValue;
 };
 
@@ -132,11 +129,6 @@ struct TUniqueFragmentSelectionSetIterator
 		++Iterator;
 	}
 
-	explicit operator bool() const
-	{
-		return static_cast<bool>(Iterator);
-	}
-
 	operator KeyType() const
 	{
 		return *Iterator;
@@ -169,8 +161,8 @@ struct TFilteredViewModelSelectionIterator
 {
 	const TSet<StorageKeyType>* SelectionSet;
 
-	FORCEINLINE TFilteredViewModelSelectionIteratorState<StorageKeyType, FilterType> begin() const { return TFilteredViewModelSelectionIteratorState<StorageKeyType, FilterType>(SelectionSet->begin()); }
-	FORCEINLINE TFilteredViewModelSelectionIteratorState<StorageKeyType, FilterType> end() const   { return TFilteredViewModelSelectionIteratorState<StorageKeyType, FilterType>(SelectionSet->end());   }
+	FORCEINLINE TFilteredViewModelSelectionIteratorState<StorageKeyType, FilterType> begin() const { return TFilteredViewModelSelectionIteratorState<StorageKeyType, FilterType>(SelectionSet->begin(), SelectionSet->end()); }
+	FORCEINLINE TFilteredViewModelSelectionIteratorState<StorageKeyType, FilterType> end() const   { return TFilteredViewModelSelectionIteratorState<StorageKeyType, FilterType>(SelectionSet->begin(), SelectionSet->end());   }
 };
 
 
@@ -182,24 +174,25 @@ struct TSelectionSetIteratorImpl<TWeakPtr<T>>
 {
 	TSharedPtr<T> operator*() const
 	{
-		return Iterator->Pin();
+		return IteratorBegin->Pin();
 	}
 
 protected:
 
-	explicit TSelectionSetIteratorImpl(typename TSet<TWeakPtr<T>>::TRangedForConstIterator InIterator)
-		: Iterator(InIterator)
+	explicit TSelectionSetIteratorImpl(typename TSet<TWeakPtr<T>>::TRangedForConstIterator InIteratorBegin, typename TSet<TWeakPtr<T>>::TRangedForConstIterator InIteratorEnd)
+		: IteratorBegin(InIteratorBegin), IteratorEnd(InIteratorEnd)
 	{}
 
 	void SkipInvalid()
 	{
-		while (Iterator && !Iterator->Pin())
+		while (IteratorBegin != IteratorEnd && !IteratorBegin->Pin())
 		{
-			++Iterator;
+			++IteratorBegin;
 		}
 	}
 
-	typename TSet<TWeakPtr<T>>::TRangedForConstIterator Iterator;
+	typename TSet<TWeakPtr<T>>::TRangedForConstIterator IteratorBegin;
+	typename TSet<TWeakPtr<T>>::TRangedForConstIterator IteratorEnd;
 };
 
 
@@ -211,24 +204,25 @@ struct TSelectionSetIteratorImpl<TWeakViewModelPtr<T>>
 {
 	TViewModelPtr<T> operator*() const
 	{
-		return Iterator->Pin();
+		return IteratorBegin->Pin();
 	}
 
 protected:
 
-	explicit TSelectionSetIteratorImpl(typename TSet<TWeakViewModelPtr<T>>::TRangedForConstIterator InIterator)
-		: Iterator(InIterator)
+	explicit TSelectionSetIteratorImpl(typename TSet<TWeakViewModelPtr<T>>::TRangedForConstIterator InIteratorBegin, typename TSet<TWeakViewModelPtr<T>>::TRangedForConstIterator InIteratorEnd)
+		: IteratorBegin(InIteratorBegin), IteratorEnd(InIteratorEnd)
 	{}
 
 	void SkipInvalid()
 	{
-		while (Iterator && !Iterator->Pin())
+		while (IteratorBegin != IteratorEnd && IteratorBegin->Pin())
 		{
-			++Iterator;
+			++IteratorBegin;
 		}
 	}
 
-	typename TSet<TWeakViewModelPtr<T>>::TRangedForConstIterator Iterator;
+	typename TSet<TWeakViewModelPtr<T>>::TRangedForConstIterator IteratorBegin;
+	typename TSet<TWeakViewModelPtr<T>>::TRangedForConstIterator IteratorEnd;
 };
 
 
diff --git a/Engine/Source/Editor/SequencerCore/Public/MVVM/Selection/SequencerCoreSelectionTypes.h b/Engine/Source/Editor/SequencerCore/Public/MVVM/Selection/SequencerCoreSelectionTypes.h
--- a/Engine/Source/Editor/SequencerCore/Public/MVVM/Selection/SequencerCoreSelectionTypes.h
+++ b/Engine/Source/Editor/SequencerCore/Public/MVVM/Selection/SequencerCoreSelectionTypes.h
@@ -292,8 +292,8 @@ public:
 public:
 
 	/** Ranged-for iteration support. Do not use directly */
-	FORCEINLINE TSelectionSetIteratorState<KeyType> begin() const { return TSelectionSetIteratorState<KeyType>(SelectionSet.begin()); }
-	FORCEINLINE TSelectionSetIteratorState<KeyType> end() const   { return TSelectionSetIteratorState<KeyType>(SelectionSet.end());   }
+	FORCEINLINE TSelectionSetIteratorState<KeyType> begin() const { return TSelectionSetIteratorState<KeyType>(SelectionSet.begin(), SelectionSet.end()); }
+	FORCEINLINE TSelectionSetIteratorState<KeyType> end() const   { return TSelectionSetIteratorState<KeyType>(SelectionSet.end(), SelectionSet.end());   }
 
 	/** Fallback functions for use when derived classes do not implement them. Should be completely compiled out */
 	FORCEINLINE static constexpr bool OnSelectItem(const KeyType&){ return true; }
diff --git a/Engine/Source/Runtime/AutomationTest/Private/AutomationTestExcludelist.cpp b/Engine/Source/Runtime/AutomationTest/Private/AutomationTestExcludelist.cpp
--- a/Engine/Source/Runtime/AutomationTest/Private/AutomationTestExcludelist.cpp
+++ b/Engine/Source/Runtime/AutomationTest/Private/AutomationTestExcludelist.cpp
@@ -264,7 +264,7 @@ FString UAutomationTestExcludelist::GetConfigFilenameForEntry(const FAutomationT
 	}
 
 	// Otherwise take the first item from the entry platform list
-	FName FirstItem = Entry.Platforms[Entry.Platforms.begin().GetId()];
+	FName FirstItem = *Entry.Platforms.FindArbitraryElement();
 	if (const TObjectPtr<UAutomationTestExcludelistConfig>* ConfigPtr = PlatformConfigs.Find(FirstItem))
 	{
 		return (*ConfigPtr)->GetConfigFilename();
diff --git a/Engine/Source/Runtime/Engine/Private/PhysicsEngine/BodyInstance.cpp b/Engine/Source/Runtime/Engine/Private/PhysicsEngine/BodyInstance.cpp
--- a/Engine/Source/Runtime/Engine/Private/PhysicsEngine/BodyInstance.cpp
+++ b/Engine/Source/Runtime/Engine/Private/PhysicsEngine/BodyInstance.cpp
@@ -1855,9 +1855,9 @@ void FBodyInstance::TermBody(bool bNeverDeferRelease)
 	{
 		TStringBuilder<256> Builder;
 		int32 BodyIndex = 0;
-		for (TMap<FPhysicsShapeHandle, FWeldInfo>::TRangedForIterator It = ShapeToBodiesMap->begin(); It; ++It)
+		for (const TPair<FPhysicsShapeHandle, FWeldInfo>& It : *ShapeToBodiesMap)
 		{
-			const FBodyInstance* ChildBody = It->Value.ChildBI;
+			const FBodyInstance* ChildBody = It.Value.ChildBI;
 			const FString Name = ChildBody != nullptr ? ChildBody->GetBodyDebugName() : FString("null");
 			Builder.Appendf(TEXT("\tBody[%d]: \'%s'\n"), BodyIndex, *Name);
 			++BodyIndex;
diff --git a/Engine/Source/Runtime/Renderer/Private/HeterogeneousVolumes/HeterogeneousVolumes.cpp b/Engine/Source/Runtime/Renderer/Private/HeterogeneousVolumes/HeterogeneousVolumes.cpp
--- a/Engine/Source/Runtime/Renderer/Private/HeterogeneousVolumes/HeterogeneousVolumes.cpp
+++ b/Engine/Source/Runtime/Renderer/Private/HeterogeneousVolumes/HeterogeneousVolumes.cpp
@@ -1683,9 +1683,9 @@ namespace HeterogeneousVolumes
 				ViewState->AdaptiveVolumetricCameraMapUniformBufferParameters = nullptr;
 				DestroyAdaptiveVolumetricShadowMapUniformBuffer(ViewState->AdaptiveVolumetricCameraMapUniformBuffer);
 
-				for (auto Itr = ViewState->AdaptiveVolumetricShadowMapUniformBufferMap.begin(); Itr != ViewState->AdaptiveVolumetricShadowMapUniformBufferMap.end(); ++Itr)
+				for (TPair<int32, TRDGUniformBufferRef<FAdaptiveVolumetricShadowMapUniformBufferParameters>>& Pair : ViewState->AdaptiveVolumetricShadowMapUniformBufferMap)
 				{
-					DestroyAdaptiveVolumetricShadowMapUniformBuffer(Itr->Value);
+					DestroyAdaptiveVolumetricShadowMapUniformBuffer(Pair.Value);
 				}
 				ViewState->AdaptiveVolumetricShadowMapUniformBufferMap.Empty();
 			}
diff --git a/Engine/Source/Runtime/Renderer/Private/SceneViewState.cpp b/Engine/Source/Runtime/Renderer/Private/SceneViewState.cpp
--- a/Engine/Source/Runtime/Renderer/Private/SceneViewState.cpp
+++ b/Engine/Source/Runtime/Renderer/Private/SceneViewState.cpp
@@ -205,11 +205,11 @@ uint64 FPreviousViewInfo::GetGPUSizeBytes(bool bLogSizes) const
 		GetRenderTargetGPUSizeBytes(MobileAmbientOcclusion, bLogSizes) +
 		GetRenderTargetGPUSizeBytes(VisualizeMotionVectors, bLogSizes);
 
-	for (auto ShadowHistoryIt = ShadowHistories.begin(); ShadowHistoryIt; ++ShadowHistoryIt)
+	for (const TPair<const ULightComponent*, TSharedPtr<FScreenSpaceDenoiserHistory>>& ShadowHistoryIt : ShadowHistories)
 	{
-		if (ShadowHistoryIt.Value().IsValid())
+		if (ShadowHistoryIt.Value.IsValid())
 		{
-			TotalSize += ShadowHistoryIt.Value()->GetGPUSizeBytes(bLogSizes);
+			TotalSize += ShadowHistoryIt.Value->GetGPUSizeBytes(bLogSizes);
 		}
 	}
 
@@ -444,11 +444,11 @@ uint64 FSceneViewState::GetGPUSizeBytes(bool bLogSizes) const
 	TotalSize += GetRenderTargetGPUSizeBytes(CombinedLUTRenderTarget, bLogSizes);
 	TotalSize += PrevFrameViewInfo.GetGPUSizeBytes(bLogSizes);
 	TotalSize += LightShaftOcclusionHistory.GetGPUSizeBytes(bLogSizes);
-	for (auto LightShaftBloomIt = LightShaftBloomHistoryRTs.begin(); LightShaftBloomIt; ++LightShaftBloomIt)
+	for (const TPair<const ULightComponent*, TUniquePtr<FTemporalAAHistory>>& LightShaftBloomIt : LightShaftBloomHistoryRTs)
 	{
-		if (LightShaftBloomIt.Value().IsValid())
+		if (LightShaftBloomIt.Value.IsValid())
 		{
-			TotalSize += LightShaftBloomIt.Value()->GetGPUSizeBytes(bLogSizes);
+			TotalSize += LightShaftBloomIt.Value->GetGPUSizeBytes(bLogSizes);
 		}
 	}
 	TotalSize += GetRenderTargetGPUSizeBytes(DistanceFieldAOHistoryRT, bLogSizes);
