Sequencer: fix iterator regressions due to 44363561

#rnx

diff --git a/Engine/Source/Editor/SequencerCore/Public/MVVM/Selection/SequencerCoreSelectionIterators.h b/Engine/Source/Editor/SequencerCore/Public/MVVM/Selection/SequencerCoreSelectionIterators.h
--- a/Engine/Source/Editor/SequencerCore/Public/MVVM/Selection/SequencerCoreSelectionIterators.h
+++ b/Engine/Source/Editor/SequencerCore/Public/MVVM/Selection/SequencerCoreSelectionIterators.h
@@ -162,7 +162,7 @@ struct TFilteredViewModelSelectionIterator
 	const TSet<StorageKeyType>* SelectionSet;
 
 	FORCEINLINE TFilteredViewModelSelectionIteratorState<StorageKeyType, FilterType> begin() const { return TFilteredViewModelSelectionIteratorState<StorageKeyType, FilterType>(SelectionSet->begin(), SelectionSet->end()); }
-	FORCEINLINE TFilteredViewModelSelectionIteratorState<StorageKeyType, FilterType> end() const   { return TFilteredViewModelSelectionIteratorState<StorageKeyType, FilterType>(SelectionSet->begin(), SelectionSet->end());   }
+	FORCEINLINE TFilteredViewModelSelectionIteratorState<StorageKeyType, FilterType> end() const   { return TFilteredViewModelSelectionIteratorState<StorageKeyType, FilterType>(SelectionSet->end(), SelectionSet->end());   }
 };
 
 
@@ -215,7 +215,7 @@ protected:
 
 	void SkipInvalid()
 	{
-		while (IteratorBegin != IteratorEnd && IteratorBegin->Pin())
+		while (IteratorBegin != IteratorEnd && !IteratorBegin->Pin())
 		{
 			++IteratorBegin;
 		}
