Sequencer: Set Mute/Solo need to refresh the tree

#rnx

diff --git a/Engine/Source/Editor/SequencerCore/Private/Scripting/OutlinerScriptingObject.cpp b/Engine/Source/Editor/SequencerCore/Private/Scripting/OutlinerScriptingObject.cpp
--- a/Engine/Source/Editor/SequencerCore/Private/Scripting/OutlinerScriptingObject.cpp
+++ b/Engine/Source/Editor/SequencerCore/Private/Scripting/OutlinerScriptingObject.cpp
@@ -183,6 +183,11 @@ void USequencerOutlinerScriptingObject::SetMuteNodes(const TArray<FSequencerView
 			MutableItem->SetIsMuted(bInMuted);
 		}
 	}
+
+	if (TViewModelPtr<FOutlinerViewModel> Outliner = WeakOutliner.Pin())
+	{
+		Outliner->RequestUpdate();
+	}
 }
 
 TArray<FSequencerViewModelScriptingStruct> USequencerOutlinerScriptingObject::GetSoloNodes()
@@ -220,6 +225,11 @@ void USequencerOutlinerScriptingObject::SetSoloNodes(const TArray<FSequencerView
 			SoloableItem->SetIsSoloed(bInSoloed);
 		}
 	}
+
+	if (TViewModelPtr<FOutlinerViewModel> Outliner = WeakOutliner.Pin())
+	{
+		Outliner->RequestUpdate();
+	}
 }
 
 TArray<FSequencerViewModelScriptingStruct> USequencerOutlinerScriptingObject::GetDeactivatedNodes()
