Pinned Section not visible in folder

Previously when dragging outliner items can be moved to a folder with a different pinned state, causing their sections not to be drawn. This change now checks if the outliner items shared the pinned state with their drop target before allowing the drop operation, which will prevent items from being added to folders that do not match their pinned state, with the appropriate UI cue of not showing the drop preview. Since multiple items can be selected, if even one item shartes the pinned state the drop will be allowed.

#Synced-CL 48956265

diff --git a/Engine/Source/Editor/SequencerCore/Private/MVVM/Views/SOutlinerViewRow.cpp b/Engine/Source/Editor/SequencerCore/Private/MVVM/Views/SOutlinerViewRow.cpp
--- a/Engine/Source/Editor/SequencerCore/Private/MVVM/Views/SOutlinerViewRow.cpp
+++ b/Engine/Source/Editor/SequencerCore/Private/MVVM/Views/SOutlinerViewRow.cpp
@@ -5,9 +5,11 @@
 #include "MVVM/Views/SOutlinerView.h"
 #include "MVVM/Views/STrackLane.h"
 #include "MVVM/Extensions/ISelectableExtension.h"
+#include "MVVM/Extensions/IPinnableExtension.h"
 #include "MVVM/ViewModels/EditorViewModel.h"
 #include "MVVM/ViewModels/OutlinerViewModel.h"
 #include "MVVM/ViewModels/EditorSharedViewModelData.h"
+#include "MVVM/ViewModels/OutlinerViewModelDragDropOp.h"
 
 
 namespace UE::Sequencer
@@ -178,6 +180,47 @@ TOptional<EItemDropZone> SOutlinerViewRow::OnCanAcceptDrop(const FDragDropEvent&
 		DropModel = DropModel->CastParent<IOutlinerDropTargetOutlinerExtension>();
 	}
 
+	// Check the pinned states of the dragged items vs the dropped items. Pinned items cannot be dragged onto unpinned items, and vice versa.
+	if(TSharedPtr<FOutlinerViewModelDragDropOp> OutlinerDragDropOp = DragDropEvent.GetOperationAs<FOutlinerViewModelDragDropOp>())
+	{
+		bool bSharedPinnedState = false;
+		
+		for (const TWeakViewModelPtr<IOutlinerExtension>& WeakDraggedModel : OutlinerDragDropOp->GetDraggedViewModels())
+		{
+			FViewModelPtr DraggedModel = WeakDraggedModel.Pin();
+			if (!DraggedModel)
+			{
+				continue;
+			}
+			if (TSharedPtr<IPinnableExtension> DraggedPinnable = DraggedModel.ImplicitCast())
+			{
+				if (TSharedPtr<IPinnableExtension> PinnableTarget = TargetModel.ImplicitCast())
+				{
+					if (DraggedPinnable->IsPinned() == PinnableTarget->IsPinned())
+					{
+						bSharedPinnedState = true;
+						break;
+					}
+				}
+				if (TSharedPtr<IPinnableExtension> PinnableDropModel = DropModel.ImplicitCast())
+				{
+					if (DraggedPinnable->IsPinned() == PinnableDropModel->IsPinned())
+					{
+						bSharedPinnedState = true;
+						break;
+					}
+				}
+			}
+			
+		}
+		
+		// If no dragged items match the pinned state of the drop target, then they cannot be dropped.
+		if (!bSharedPinnedState)
+		{
+			return TOptional<EItemDropZone>();
+		}
+	}
+
 	TViewModelPtr<IOutlinerDropTargetOutlinerExtension> DropTarget = DropModel.ImplicitCast();
 	if (!DropTarget)
 	{
