Sequencer: Anim Mixer. Fix issue where you can't drag a layer in the outliner under an expanded layer.

This is done by allowing the drag operation to set the drop model and target model.
Also sneaky fix to prevent accidental non-renames of layers from sticking.

#Synced-CL 49622961

diff --git a/Engine/Plugins/Experimental/MovieSceneAnimMixer/Source/MovieSceneAnimMixer/Private/MovieSceneAnimationMixerLayer.cpp b/Engine/Plugins/Experimental/MovieSceneAnimMixer/Source/MovieSceneAnimMixer/Private/MovieSceneAnimationMixerLayer.cpp
--- a/Engine/Plugins/Experimental/MovieSceneAnimMixer/Source/MovieSceneAnimMixer/Private/MovieSceneAnimationMixerLayer.cpp
+++ b/Engine/Plugins/Experimental/MovieSceneAnimMixer/Source/MovieSceneAnimMixer/Private/MovieSceneAnimationMixerLayer.cpp
@@ -102,7 +102,7 @@ FText UMovieSceneAnimationMixerLayer::GetDisplayName() const
 
 void UMovieSceneAnimationMixerLayer::SetDisplayName(const FText& InDisplayName)
 {
-	if (!DisplayName.EqualTo(InDisplayName))
+	if (!GetDisplayName().EqualTo(InDisplayName))
 	{
 		DisplayName = InDisplayName;
 		MarkAsChanged();
diff --git a/Engine/Source/Editor/Sequencer/Private/MVVM/ViewModels/TrackRowReorderDragDropOp.cpp b/Engine/Source/Editor/Sequencer/Private/MVVM/ViewModels/TrackRowReorderDragDropOp.cpp
--- a/Engine/Source/Editor/Sequencer/Private/MVVM/ViewModels/TrackRowReorderDragDropOp.cpp
+++ b/Engine/Source/Editor/Sequencer/Private/MVVM/ViewModels/TrackRowReorderDragDropOp.cpp
@@ -203,4 +203,18 @@ bool FTrackRowReorderDragDropOp::ExecuteReorder(int32 TargetRowIndex)
 	return bAnyChanged;
 }
 
+void FTrackRowReorderDragDropOp::GetTargetAndDropModel(const FViewModelPtr& CurrentViewModel, EItemDropZone InItemDropZone, FViewModelPtr& DropModel, FViewModelPtr& TargetModel) const
+{
+	// Use the base implementation first to find the current level's target model and drop model
+	FOutlinerViewModelDragDropOp::GetTargetAndDropModel(CurrentViewModel, InItemDropZone, DropModel, TargetModel);
+
+	// If the TargetModel is not a track row and we're trying to drop after it, change the target model to the parent and go again
+	
+	TViewModelPtr<ITrackRowExtension> TrackRow = TargetModel.ImplicitCast();
+	if (!TrackRow && TargetModel && !TargetModel->GetNextSibling() && InItemDropZone == EItemDropZone::BelowItem)
+	{
+		GetTargetAndDropModel(TargetModel->GetParent(), InItemDropZone, DropModel, TargetModel);
+	}
+}
+
 } // namespace UE::Sequencer
diff --git a/Engine/Source/Editor/Sequencer/Public/MVVM/ViewModels/TrackRowReorderDragDropOp.h b/Engine/Source/Editor/Sequencer/Public/MVVM/ViewModels/TrackRowReorderDragDropOp.h
--- a/Engine/Source/Editor/Sequencer/Public/MVVM/ViewModels/TrackRowReorderDragDropOp.h
+++ b/Engine/Source/Editor/Sequencer/Public/MVVM/ViewModels/TrackRowReorderDragDropOp.h
@@ -45,6 +45,12 @@ public:
 	 */
 	SEQUENCER_API bool ExecuteReorder(int32 TargetRowIndex);
 
+	/*
+	* We override this to ensure the drop model is always the track and the target is always a track row even if we're dragging below
+	* a nested channel.
+	*/
+	SEQUENCER_API virtual void GetTargetAndDropModel(const FViewModelPtr& CurrentViewModel, EItemDropZone InItemDropZone, FViewModelPtr& DropModel, FViewModelPtr& TargetModel) const override;
+
 protected:
 	/** Protected construction - use New() factory function */
 	FTrackRowReorderDragDropOp()
diff --git a/Engine/Source/Editor/SequencerCore/Private/MVVM/ViewModels/OutlinerViewModelDragDropOp.cpp b/Engine/Source/Editor/SequencerCore/Private/MVVM/ViewModels/OutlinerViewModelDragDropOp.cpp
--- a/Engine/Source/Editor/SequencerCore/Private/MVVM/ViewModels/OutlinerViewModelDragDropOp.cpp
+++ b/Engine/Source/Editor/SequencerCore/Private/MVVM/ViewModels/OutlinerViewModelDragDropOp.cpp
@@ -94,6 +94,18 @@ bool FOutlinerViewModelDragDropOp::ValidateParentChildDrop(const FViewModel& Pro
 	return true;
 }
 
+void FOutlinerViewModelDragDropOp::GetTargetAndDropModel(const FViewModelPtr& CurrentViewModel, EItemDropZone InItemDropZone, FViewModelPtr& DropModel, FViewModelPtr& TargetModel) const
+{
+	DropModel = CurrentViewModel;
+
+	// When dropping above or below an item, we always forward to its parent
+	if (InItemDropZone == EItemDropZone::AboveItem || InItemDropZone == EItemDropZone::BelowItem)
+	{
+		TargetModel = DropModel;
+		DropModel = DropModel->CastParent<IOutlinerDropTargetOutlinerExtension>();
+	}
+}
+
 } // namespace Sequencer
 } // namespace UE
 
diff --git a/Engine/Source/Editor/SequencerCore/Private/MVVM/Views/SOutlinerViewRow.cpp b/Engine/Source/Editor/SequencerCore/Private/MVVM/Views/SOutlinerViewRow.cpp
--- a/Engine/Source/Editor/SequencerCore/Private/MVVM/Views/SOutlinerViewRow.cpp
+++ b/Engine/Source/Editor/SequencerCore/Private/MVVM/Views/SOutlinerViewRow.cpp
@@ -163,25 +163,25 @@ TOptional<EItemDropZone> SOutlinerViewRow::OnCanAcceptDrop(const FDragDropEvent&
 		return TOptional<EItemDropZone>();
 	}
 
-	// The model we are indirectly dropping into (either this model, or its parent)
+	// The model we are indirectly dropping into
 	FViewModelPtr DropModel = ThisModel;
-	// The model we are directly interacting with - only set when attaching before/after
+	// The model we are directly interacting with
 	FViewModelPtr TargetModel;
 
-	if (!DropModel)
+	TSharedPtr<FOutlinerViewModelDragDropOp> OutlinerDragDropOp = DragDropEvent.GetOperationAs<FOutlinerViewModelDragDropOp>();
+	if (OutlinerDragDropOp)
 	{
-		return TOptional<EItemDropZone>();
+		// Allow outliner drag operations to override drop model and target model
+		OutlinerDragDropOp->GetTargetAndDropModel(ThisModel, InItemDropZone, DropModel, TargetModel);
 	}
-
-	// When dropping above or below an item, we always forward to its parent
-	if (InItemDropZone == EItemDropZone::AboveItem || InItemDropZone == EItemDropZone::BelowItem)
+	else if (InItemDropZone == EItemDropZone::AboveItem || InItemDropZone == EItemDropZone::BelowItem)
 	{
 		TargetModel = DropModel;
 		DropModel = DropModel->CastParent<IOutlinerDropTargetOutlinerExtension>();
 	}
 
 	// Check the pinned states of the dragged items vs the dropped items. Pinned items cannot be dragged onto unpinned items, and vice versa.
-	if(TSharedPtr<FOutlinerViewModelDragDropOp> OutlinerDragDropOp = DragDropEvent.GetOperationAs<FOutlinerViewModelDragDropOp>())
+	if(OutlinerDragDropOp)
 	{
 		bool bSharedPinnedState = false;
 		
@@ -245,13 +245,17 @@ FReply SOutlinerViewRow::OnAcceptDrop(const FDragDropEvent& DragDropEvent, EItem
 		return FReply::Unhandled();
 	}
 
-	// The model we are indirectly dropping into (either this model, or its parent)
+	// The model we are indirectly dropping into
 	FViewModelPtr DropModel = ThisModel;
-	// The model we are directly interacting with - only set when attaching before/after
+	// The model we are directly interacting with
 	FViewModelPtr TargetModel;
 
-	// When dropping above or below an item, we always forward to its parent
-	if (InItemDropZone == EItemDropZone::AboveItem || InItemDropZone == EItemDropZone::BelowItem)
+	if (TSharedPtr<FOutlinerViewModelDragDropOp> OutlinerDragDropOp = DragDropEvent.GetOperationAs<FOutlinerViewModelDragDropOp>())
+	{
+		// Allow outliner drag operations to override drop model and target model
+		OutlinerDragDropOp->GetTargetAndDropModel(ThisModel, InItemDropZone, DropModel, TargetModel);
+	}
+	else if (InItemDropZone == EItemDropZone::AboveItem || InItemDropZone == EItemDropZone::BelowItem)
 	{
 		TargetModel = DropModel;
 		DropModel = DropModel->CastParent<IOutlinerDropTargetOutlinerExtension>();
diff --git a/Engine/Source/Editor/SequencerCore/Public/MVVM/ViewModels/OutlinerViewModelDragDropOp.h b/Engine/Source/Editor/SequencerCore/Public/MVVM/ViewModels/OutlinerViewModelDragDropOp.h
--- a/Engine/Source/Editor/SequencerCore/Public/MVVM/ViewModels/OutlinerViewModelDragDropOp.h
+++ b/Engine/Source/Editor/SequencerCore/Public/MVVM/ViewModels/OutlinerViewModelDragDropOp.h
@@ -18,6 +18,7 @@
 #define UE_API SEQUENCERCORE_API
 
 struct FSlateBrush;
+enum class EItemDropZone;
 
 namespace UE
 {
@@ -69,6 +70,15 @@ public:
 	 */
 	UE_API bool ValidateParentChildDrop(const FViewModel& ProspectiveItem);
 
+	/*
+	* Based on the current view model being hovered, return the 'target' model and the 'drop model'. 
+	* DropModel is the model we will call to execute the drop. 
+	* TargetModel is only set if we're inserting before or after a child model of DropModel. 
+	* We default to using the CurrentViewModel for the DropModel unless we've dragged above or below the CurrentViewModel,
+	* in which we default to the CurrentViewModel's parent for the DropModel and the TargetModel as the CurrentViewModel.
+	*/
+	UE_API virtual void GetTargetAndDropModel(const FViewModelPtr& CurrentViewModel, EItemDropZone InItemDropZone, FViewModelPtr& DropModel, FViewModelPtr& TargetModel) const;
+
 private:
 
 	void AddSnapTime(FFrameNumber SnapTime) override
