Sequencer: Remove column buttons for transition section outliners.

#Synced-CL 51964473

diff --git a/Engine/Plugins/Experimental/MovieSceneAnimMixer/Source/MovieSceneAnimMixerEditor/Private/AnimTransitionSectionInterface.h b/Engine/Plugins/Experimental/MovieSceneAnimMixer/Source/MovieSceneAnimMixerEditor/Private/AnimTransitionSectionInterface.h
--- a/Engine/Plugins/Experimental/MovieSceneAnimMixer/Source/MovieSceneAnimMixerEditor/Private/AnimTransitionSectionInterface.h
+++ b/Engine/Plugins/Experimental/MovieSceneAnimMixer/Source/MovieSceneAnimMixerEditor/Private/AnimTransitionSectionInterface.h
@@ -45,6 +45,9 @@ public:
 	/** Transition sections cannot be explicitly deleted - they are removed automatically when overlap goes away */
 	virtual bool IsDeletable() const override { return false; }
 
+	/** Transition sections should not expose mute/solo/lock/deactivate/condition column options */
+	virtual bool SupportsOutlinerColumnToggle(FName ColumnName) const override { return false; }
+
 	/** Transition sections draw their own curve visualization in OnPaintSection, so don't show the key area when collapsed */
 	virtual bool ShouldShowKeyAreaWhenCollapsed() const override { return false; }
 
diff --git a/Engine/Source/Editor/Sequencer/Private/MVVM/ViewModels/SectionOutlinerModel.cpp b/Engine/Source/Editor/Sequencer/Private/MVVM/ViewModels/SectionOutlinerModel.cpp
--- a/Engine/Source/Editor/Sequencer/Private/MVVM/ViewModels/SectionOutlinerModel.cpp
+++ b/Engine/Source/Editor/Sequencer/Private/MVVM/ViewModels/SectionOutlinerModel.cpp
@@ -322,6 +322,18 @@ TSharedPtr<SWidget> FSectionOutlinerModel::CreateOutlinerViewForColumn(const FCr
 	return nullptr;
 }
 
+bool FSectionOutlinerModel::SupportsOutlinerColumn(const FName& InColumnName) const
+{
+	if (TSharedPtr<FSectionModel> SectionModel = GetSectionModel())
+	{
+		if (TSharedPtr<ISequencerSection> SectionInterface = SectionModel->GetSectionInterface())
+		{
+			return SectionInterface->SupportsOutlinerColumnToggle(InColumnName);
+		}
+	}
+	return true;
+}
+
 FText FSectionOutlinerModel::GetLabel() const
 {
 	// Try to get the title from the associated FSectionModel's ISequencerSection interface
diff --git a/Engine/Source/Editor/Sequencer/Public/ISequencerSection.h b/Engine/Source/Editor/Sequencer/Public/ISequencerSection.h
--- a/Engine/Source/Editor/Sequencer/Public/ISequencerSection.h
+++ b/Engine/Source/Editor/Sequencer/Public/ISequencerSection.h
@@ -257,6 +257,14 @@ public:
 	 */
 	virtual bool IsDeletable() const { return true; }
 
+	/**
+	 * Controls whether a given outliner column toggle (mute, solo, lock, deactivate, etc.)
+	 * is available for this section's outliner item. Default is true for all columns.
+	 * Override to return false for specific columns on sections like transitions that should
+	 * not expose certain controls.
+	 */
+	virtual bool SupportsOutlinerColumnToggle(FName ColumnName) const { return true; }
+
 	/**
 	 * Controls whether the section's key area should be visible when the section's parent
 	 * outliner item is collapsed. Default is true.
diff --git a/Engine/Source/Editor/Sequencer/Public/MVVM/ViewModels/SectionOutlinerModel.h b/Engine/Source/Editor/Sequencer/Public/MVVM/ViewModels/SectionOutlinerModel.h
--- a/Engine/Source/Editor/Sequencer/Public/MVVM/ViewModels/SectionOutlinerModel.h
+++ b/Engine/Source/Editor/Sequencer/Public/MVVM/ViewModels/SectionOutlinerModel.h
@@ -87,6 +87,7 @@ public:
 	/*~ IOutlinerExtension */
 	UE_API FOutlinerSizing GetOutlinerSizing() const override;
 	UE_API TSharedPtr<SWidget> CreateOutlinerViewForColumn(const FCreateOutlinerViewParams& InParams, const FName& InColumnName) override;
+	UE_API bool SupportsOutlinerColumn(const FName& InColumnName) const override;
 	UE_API FText GetLabel() const override;
 	UE_API FSlateColor GetLabelColor() const override;
 	UE_API FSlateFontInfo GetLabelFont() const override;
diff --git a/Engine/Source/Editor/SequencerCore/Private/MVVM/Views/SOutlinerView.cpp b/Engine/Source/Editor/SequencerCore/Private/MVVM/Views/SOutlinerView.cpp
--- a/Engine/Source/Editor/SequencerCore/Private/MVVM/Views/SOutlinerView.cpp
+++ b/Engine/Source/Editor/SequencerCore/Private/MVVM/Views/SOutlinerView.cpp
@@ -306,8 +306,8 @@ int32 SOutlinerView::CreateOutlinerColumnsForGroup(int32 ColumnIndex, EOutlinerC
 
 		// Keep track of how to create widgets for this column
 		// Column Generators must be defined before HeaderRow->AddColumn since AddColumn might re-generate widgets
-		ColumnGenerators.Add(ColumnName, [Column](const FCreateOutlinerColumnParams& Params, const TSharedRef<SOutlinerViewRow>& Row){
-			return Column->IsItemCompatibleWithColumn(Params)
+		ColumnGenerators.Add(ColumnName, [Column, ColumnName](const FCreateOutlinerColumnParams& Params, const TSharedRef<SOutlinerViewRow>& Row){
+			return (Params.OutlinerExtension.IsValid() && Params.OutlinerExtension->SupportsOutlinerColumn(ColumnName) && Column->IsItemCompatibleWithColumn(Params))
 				? Column->CreateColumnWidget(Params, Row)
 				: nullptr;
 		});
diff --git a/Engine/Source/Editor/SequencerCore/Public/MVVM/Extensions/IOutlinerExtension.h b/Engine/Source/Editor/SequencerCore/Public/MVVM/Extensions/IOutlinerExtension.h
--- a/Engine/Source/Editor/SequencerCore/Public/MVVM/Extensions/IOutlinerExtension.h
+++ b/Engine/Source/Editor/SequencerCore/Public/MVVM/Extensions/IOutlinerExtension.h
@@ -190,6 +190,10 @@ public:
 	/** Create the visual outliner widget for the specified column */
 	UE_API virtual TSharedPtr<SWidget> CreateOutlinerViewForColumn(const FCreateOutlinerViewParams& InParams, const FName& InColumnName);
 
+	/** Whether this outliner item supports the given column. Default is true for all columns.
+	 *  Override to suppress specific columns (e.g. mute/solo on transition sections). */
+	virtual bool SupportsOutlinerColumn(const FName& InColumnName) const { return true; }
+
 	/** Build the context menu widget for this outliner item */
 	virtual TSharedPtr<SWidget> CreateContextMenuWidget(const FCreateOutlinerContextMenuWidgetParams& InParams) = 0;
 
