[Sequencer] Simple View with toolable timeline initial release

diff --git a/Engine/Source/Editor/Sequencer/Private/ToolableTimeline/ToolableTimelineKeySelection.cpp b/Engine/Source/Editor/Sequencer/Private/ToolableTimeline/ToolableTimelineKeySelection.cpp
new file mode 100644
--- /dev/null
+++ b/Engine/Source/Editor/Sequencer/Private/ToolableTimeline/ToolableTimelineKeySelection.cpp
@@ -0,0 +1,687 @@
+﻿// Copyright Epic Games, Inc. All Rights Reserved.
+
+#include "ToolableTimelineKeySelection.h"
+#include "CurveEditor.h"
+#include "CurveEditorSelection.h"
+#include "ISequencer.h"
+#include "MouseInputData.h"
+#include "MVVM/Selection/Selection.h"
+#include "Sequencer.h"
+#include "ToolableTimeline/ToolableTimelineUtils.h"
+#include "ToolableTimeSliderController.h"
+
+namespace UE::Sequencer::ToolableTimeline
+{
+
+FToolableTimelineKeySelection::FToolableTimelineKeySelection(FToolableTimeline& InTimeline)
+	: Timeline(InTimeline)
+{
+}
+
+FToolableTimelineKeySelection::~FToolableTimelineKeySelection()
+{
+	Shutdown();
+}
+
+void FToolableTimelineKeySelection::Initialize()
+{
+	if (const TSharedPtr<FSequencer> Sequencer = Timeline.GetSequencer())
+	{
+		FSequencerSelection& SequencerSelection = Sequencer->GetSelection();
+		SequencerKeySelectionChangedDelegate = SequencerSelection.KeySelection.OnChanged.AddRaw(this, &FToolableTimelineKeySelection::HandleSequencerKeySelectionChanged);
+
+		if (const TSharedPtr<FCurveEditor> CurveEditor = Utils::GetSequencerCurveEditor(*Sequencer))
+		{
+			CurveEditorKeySelectionChangedDelegate = CurveEditor->GetSelection().OnSelectionChanged().AddRaw(this, &FToolableTimelineKeySelection::HandleCurveEditorKeySelectionChanged);
+		}
+	}
+}
+
+void FToolableTimelineKeySelection::Shutdown()
+{
+	if (const TSharedPtr<FSequencer> Sequencer = Timeline.GetSequencer())
+	{
+		if (const TSharedPtr<FCurveEditor> CurveEditor = Utils::GetSequencerCurveEditor(*Sequencer))
+		{
+			CurveEditor->GetSelection().OnSelectionChanged().Remove(CurveEditorKeySelectionChangedDelegate);
+			CurveEditorKeySelectionChangedDelegate.Reset();
+		}
+
+		FSequencerSelection& SequencerSelection = Sequencer->GetSelection();
+		SequencerSelection.KeySelection.OnChanged.Remove(SequencerKeySelectionChangedDelegate);
+		SequencerKeySelectionChangedDelegate.Reset();
+	}
+}
+
+void FToolableTimelineKeySelection::HandleSequencerKeySelectionChanged()
+{
+	if (bSyncDisabled || !IsSimpleView())
+	{
+		return;
+	}
+	TGuardValue<bool> Guard(bSyncDisabled, true);
+
+	const TSet<FSequencerSelectedKey> NewSelectedKeys = BuildSelectedKeysFromSequencer();
+
+	if (NewSelectedKeys.Num() == SelectedKeys.Num() && NewSelectedKeys.Includes(SelectedKeys))
+	{
+		return;
+	}
+
+	SelectedKeys = NewSelectedKeys;
+
+	ClearHoveredKeys(false);
+	ClearSelectionPreview();
+
+	SyncKeySelectionToCurveEditor();
+}
+
+void FToolableTimelineKeySelection::HandleCurveEditorKeySelectionChanged()
+{
+	if (bSyncDisabled || !IsSimpleView())
+	{
+		return;
+	}
+	TGuardValue<bool> Guard(bSyncDisabled, true);
+
+	const TSet<FSequencerSelectedKey> NewSelectedKeys = BuildSelectedKeysFromCurveEditor();
+
+	if (NewSelectedKeys.Num() == SelectedKeys.Num() && NewSelectedKeys.Includes(SelectedKeys))
+	{
+		return;
+	}
+
+	SelectedKeys = NewSelectedKeys;
+
+	ClearHoveredKeys(false);
+	ClearSelectionPreview();
+
+	SyncKeySelectionToSequencer();
+}
+
+TSet<FSequencerSelectedKey> FToolableTimelineKeySelection::BuildSelectedKeysFromSequencer() const
+{
+	TSet<FSequencerSelectedKey> OutKeys;
+
+	const TSharedPtr<FSequencer> Sequencer = Timeline.GetSequencer();
+	if (!Sequencer.IsValid())
+	{
+		return OutKeys;
+	}
+
+	const FKeySelection& SequencerKeySelection = Sequencer->GetSelection().KeySelection;
+
+	for (const FKeyHandle KeyHandle : SequencerKeySelection)
+	{
+		if (const TViewModelPtr<FChannelModel> KeyModel = SequencerKeySelection.GetModelForKey(KeyHandle))
+		{
+			FSequencerSelectedKey SelectedKey;
+			SelectedKey.WeakChannel = KeyModel;
+			SelectedKey.Section = KeyModel->GetSection();
+			SelectedKey.KeyHandle = KeyHandle;
+			OutKeys.Add(SelectedKey);
+		}
+	}
+
+	return OutKeys;
+}
+
+TSet<FSequencerSelectedKey> FToolableTimelineKeySelection::BuildSelectedKeysFromCurveEditor() const
+{
+	TSet<FSequencerSelectedKey> OutKeys;
+
+	const TSharedPtr<ISequencer> Sequencer = Timeline.GetSequencer();
+	if (!Sequencer.IsValid())
+	{
+		return OutKeys;
+	}
+
+	const TSharedPtr<FCurveEditor> CurveEditor = Utils::GetSequencerCurveEditor(*Sequencer);
+	if (!CurveEditor.IsValid())
+	{
+		return OutKeys;
+	}
+
+	const TMap<FCurveModelID, FKeyHandleSet>& AllSelectedCurveKeys = CurveEditor->GetSelection().GetAll();
+
+	for (const TPair<FCurveModelID, FKeyHandleSet>& Pair : AllSelectedCurveKeys)
+	{
+		FCurveModel* const CurveModel = CurveEditor->FindCurve(Pair.Key);
+		if (!CurveModel)
+		{
+			continue;
+		}
+
+		const TSharedPtr<FChannelModel> ChannelModel = Utils::FindChannelModelForCurve(Timeline.AsShared(), *CurveModel);
+		if (!ChannelModel.IsValid())
+		{
+			continue;
+		}
+
+		for (const FKeyHandle KeyHandle : Pair.Value.AsArray())
+		{
+			FSequencerSelectedKey SelectedKey;
+			SelectedKey.WeakChannel = ChannelModel;
+			SelectedKey.Section = ChannelModel->GetSection();
+			SelectedKey.KeyHandle = KeyHandle;
+			OutKeys.Add(SelectedKey);
+		}
+	}
+
+	return OutKeys;
+}
+
+bool FToolableTimelineKeySelection::IsSequencerSelectionEqualToSelectedKeys() const
+{
+	const TSet<FSequencerSelectedKey> ExistingKeys = BuildSelectedKeysFromSequencer();
+	return ExistingKeys.Num() == SelectedKeys.Num() && ExistingKeys.Includes(SelectedKeys);
+}
+
+bool FToolableTimelineKeySelection::IsCurveEditorSelectionEqualToSelectedKeys() const
+{
+	const TSet<FSequencerSelectedKey> ExistingKeys = BuildSelectedKeysFromCurveEditor();
+	return ExistingKeys.Num() == SelectedKeys.Num() && ExistingKeys.Includes(SelectedKeys);
+}
+
+bool FToolableTimelineKeySelection::IsSelectionPreviewEqual(const TMap<FKeyHandle, ESelectionPreviewState>& InPreview) const
+{
+	if (SelectionPreview.Num() != InPreview.Num())
+	{
+		return false;
+	}
+
+	for (const TPair<FKeyHandle, ESelectionPreviewState>& Pair : SelectionPreview)
+	{
+		const ESelectionPreviewState* OtherState = InPreview.Find(Pair.Key);
+		if (!OtherState || *OtherState != Pair.Value)
+		{
+			return false;
+		}
+	}
+
+	return true;
+}
+
+bool FToolableTimelineKeySelection::HasAnySelectedKeys() const
+{
+	return !SelectedKeys.IsEmpty();
+}
+
+TSet<FSequencerSelectedKey> FToolableTimelineKeySelection::GetSelectedKeys() const
+{
+	return SelectedKeys;
+}
+
+bool FToolableTimelineKeySelection::ClearSelectedKeys(const bool bInSync)
+{
+	if (SelectedKeys.IsEmpty())
+	{
+		return false;
+	}
+
+	SelectedKeys.Reset();
+
+	Timeline.GetTimeSliderController()->InvalidateKeyRendererCache();
+
+	if (bInSync)
+	{
+		SyncSelectionToSequencerAndCurveEditor();
+	}
+
+	return true;
+}
+
+bool FToolableTimelineKeySelection::SetSelectedKeys(const TSet<FSequencerSelectedKey>& InKeys, const bool bInSync)
+{
+	if (SelectedKeys.Num() == InKeys.Num() && SelectedKeys.Includes(InKeys))
+	{
+		return false;
+	}
+
+	SelectedKeys = InKeys;
+
+	Timeline.GetTimeSliderController()->InvalidateKeyRendererCache();
+
+	if (bInSync)
+	{
+		SyncSelectionToSequencerAndCurveEditor();
+	}
+
+	return true;
+}
+
+bool FToolableTimelineKeySelection::AddSelectedKeys(const TSet<FSequencerSelectedKey>& InKeys, const bool bInSync)
+{
+	const int32 OldNum = SelectedKeys.Num();
+
+	for (const FSequencerSelectedKey& Key : InKeys)
+	{
+		SelectedKeys.Add(Key);
+	}
+
+	if (SelectedKeys.Num() == OldNum)
+	{
+		return false;
+	}
+
+	Timeline.GetTimeSliderController()->InvalidateKeyRendererCache();
+
+	if (bInSync)
+	{
+		SyncSelectionToSequencerAndCurveEditor();
+	}
+
+	return true;
+}
+
+bool FToolableTimelineKeySelection::HasAnyHoveredKeys() const
+{
+	return !HoveredKeys.IsEmpty();
+}
+
+TSet<FSequencerSelectedKey> FToolableTimelineKeySelection::GetHoveredKeys() const
+{
+	return HoveredKeys;
+}
+
+bool FToolableTimelineKeySelection::ClearHoveredKeys(const bool bInSync)
+{
+	if (HoveredKeys.IsEmpty())
+	{
+		return false;
+	}
+
+	HoveredKeys.Reset();
+
+	Timeline.GetTimeSliderController()->InvalidateKeyRendererCache();
+
+	if (bInSync)
+	{
+		SyncSelectionToSequencerAndCurveEditor();
+	}
+
+	return true;
+}
+
+bool FToolableTimelineKeySelection::SetHoveredKeys(const TSet<FSequencerSelectedKey>& InKeys)
+{
+	if (HoveredKeys.Num() == InKeys.Num() && HoveredKeys.Includes(InKeys))
+	{
+		return false;
+	}
+
+	HoveredKeys = InKeys;
+
+	Timeline.GetTimeSliderController()->InvalidateKeyRendererCache();
+
+	return true;
+}
+
+bool FToolableTimelineKeySelection::ClearSelectedAndHoveredKeys(const bool bInSync)
+{
+	bool bChanged = false;
+
+	bChanged |= ClearSelectedKeys(false);
+	bChanged |= ClearHoveredKeys(false);
+	bChanged |= ClearSelectionPreview();
+
+	if (bChanged && bInSync)
+	{
+		SyncSelectionToSequencerAndCurveEditor();
+	}
+
+	return bChanged;
+}
+
+bool FToolableTimelineKeySelection::SetSelectedAndHoveredKeys(const TSet<FSequencerSelectedKey>& InSelectedKeys, const bool bInSync)
+{
+	bool bChanged = false;
+
+	bChanged |= SetSelectedKeys(InSelectedKeys, false);
+	bChanged |= SetHoveredKeys(InSelectedKeys);
+	bChanged |= ClearSelectionPreview();
+
+	if (bChanged && bInSync)
+	{
+		SyncSelectionToSequencerAndCurveEditor();
+	}
+
+	return bChanged;
+}
+
+bool FToolableTimelineKeySelection::HasAnyPreviewSelectedKeys() const
+{
+	return !SelectionPreview.IsEmpty();
+}
+
+bool FToolableTimelineKeySelection::ClearSelectionPreview()
+{
+	if (SelectionPreview.IsEmpty())
+	{
+		return false;
+	}
+
+	SelectionPreview.Reset();
+
+	Timeline.GetTimeSliderController()->InvalidateKeyRendererCache();
+
+	return true;
+}
+
+bool FToolableTimelineKeySelection::SetSelectionPreview(const TMap<FKeyHandle, ESelectionPreviewState>& InPreview)
+{
+	if (IsSelectionPreviewEqual(InPreview))
+	{
+		return false;
+	}
+
+	SelectionPreview = InPreview;
+
+	Timeline.GetTimeSliderController()->InvalidateKeyRendererCache();
+
+	return true;
+}
+
+bool FToolableTimelineKeySelection::ResetHoveredKeysAndSelectionPreview()
+{
+	bool bChanged = false;
+
+	bChanged |= ClearHoveredKeys(false);
+	bChanged |= ClearSelectionPreview();
+
+	return bChanged;
+}
+
+bool FToolableTimelineKeySelection::IsKeySelected(const FSequencerSelectedKey& InKey) const
+{
+	return SelectedKeys.Contains(InKey);
+}
+
+bool FToolableTimelineKeySelection::IsKeySelected(const TViewModelPtr<IKeyExtension>& InOwner, const FKeyHandle InKey) const
+{
+	for (const FSequencerSelectedKey& Key : SelectedKeys)
+	{
+		if (Key.WeakChannel.Pin() == InOwner)
+		{
+			return Key.KeyHandle == InKey;
+		}
+	}
+	return false;
+}
+
+bool FToolableTimelineKeySelection::IsKeyHovered(const FSequencerSelectedKey& InKey) const
+{
+	return HoveredKeys.Contains(InKey);
+}
+
+bool FToolableTimelineKeySelection::IsKeyHovered(const TViewModelPtr<IKeyExtension>& InOwner, const FKeyHandle InKey) const
+{
+	for (const FSequencerSelectedKey& Key : HoveredKeys)
+	{
+		if (Key.WeakChannel.Pin() == InOwner)
+		{
+			return Key.KeyHandle == InKey;
+		}
+	}
+	return false;
+}
+
+ESelectionPreviewState FToolableTimelineKeySelection::GetKeyPreviewSelectionState(const FKeyHandle InKey) const
+{
+	const ESelectionPreviewState* KeySelectionPreview = SelectionPreview.Find(InKey);
+	return KeySelectionPreview ? *KeySelectionPreview : ESelectionPreviewState::Undefined;
+}
+
+EKeySelectionPreviewState FToolableTimelineKeySelection::GetPreviewSelectionState(const TViewModelPtr<IKeyExtension>& InOwner, const FKeyHandle InKey) const
+{
+	const ESelectionPreviewState KeyPreviewSelectionState = GetKeyPreviewSelectionState(InKey);
+	switch(KeyPreviewSelectionState)
+	{
+	default:
+	case ESelectionPreviewState::Undefined: return EKeySelectionPreviewState::Undefined;
+	case ESelectionPreviewState::Selected: return EKeySelectionPreviewState::Selected;
+	case ESelectionPreviewState::NotSelected: return EKeySelectionPreviewState::NotSelected;
+	}
+}
+
+bool FToolableTimelineKeySelection::SelectKeys(const TSet<FSequencerSelectedKey>& InKeys, const bool bInAddToSelection)
+{
+	bool bSelectionChanged = false;
+
+	if (InKeys.IsEmpty())
+	{
+		bSelectionChanged |= ClearSelectedKeys(false);
+		bSelectionChanged |= ClearHoveredKeys(false);
+		bSelectionChanged |= ClearSelectionPreview();
+	}
+	else
+	{
+		if (bInAddToSelection)
+		{
+			AddSelectedKeys(InKeys, false);
+		}
+		else
+		{
+			SetSelectedKeys(InKeys, false);
+		}
+
+		SetHoveredKeys(InKeys);
+
+		bSelectionChanged = true;
+	}
+
+	if (bSelectionChanged)
+	{
+		Timeline.SyncSelectionToSequencerAndCurveEditor();
+	}
+
+	return bSelectionChanged;
+}
+
+bool FToolableTimelineKeySelection::SelectKeysUnderMouse(const FMouseInputData& InMouseInput, const bool bInAddToSelection)
+{
+	const TSharedRef<FToolableTimeSliderController> TimeSliderController = InMouseInput.Timeline->GetTimeSliderController();
+
+	TSet<FSequencerSelectedKey> KeysUnderMouse;
+	TimeSliderController->GetKeysUnderMouse(InMouseInput, InMouseInput.PointerEvent.GetScreenSpacePosition(), KeysUnderMouse);
+
+	return SelectKeys(KeysUnderMouse, bInAddToSelection);
+}
+
+void FToolableTimelineKeySelection::SyncKeySelectionToSequencer()
+{
+	if (!IsSimpleView())
+	{
+		return;
+	}
+	TGuardValue<bool> Guard(bSyncDisabled, true);
+
+	if (IsSequencerSelectionEqualToSelectedKeys())
+	{
+		return;
+	}
+
+	const TSharedPtr<FSequencer> Sequencer = Timeline.GetSequencer();
+	if (!Sequencer.IsValid())
+	{
+		return;
+	}
+
+	FKeySelection& KeySelection = Sequencer->GetSelection().KeySelection;
+
+	KeySelection.Empty();
+
+	for (const FSequencerSelectedKey& SelectedKey : SelectedKeys)
+	{
+		if (const TSharedPtr<FChannelModel> ChannelModel = SelectedKey.WeakChannel.Pin())
+		{
+			KeySelection.Select(ChannelModel, SelectedKey.KeyHandle);
+		}
+	}
+}
+
+void FToolableTimelineKeySelection::SyncKeySelectionToCurveEditor()
+{
+	if (!IsSimpleView())
+	{
+		return;
+	}
+	TGuardValue<bool> Guard(bSyncDisabled, true);
+
+	if (IsCurveEditorSelectionEqualToSelectedKeys())
+	{
+		return;
+	}
+
+	const TSharedPtr<ISequencer> Sequencer = Timeline.GetSequencer();
+	if (!Sequencer.IsValid())
+	{
+		return;
+	}
+
+	const TSharedPtr<FCurveEditor> CurveEditor = Utils::GetSequencerCurveEditor(*Sequencer);
+	if (!CurveEditor.IsValid())
+	{
+		return;
+	}
+
+	// Always ensure the curve editor key selection change delegate is bound since it may not be available on initialization
+	if (!CurveEditorKeySelectionChangedDelegate.IsValid())
+	{
+		CurveEditorKeySelectionChangedDelegate = CurveEditor->GetSelection().OnSelectionChanged().AddRaw(this, &FToolableTimelineKeySelection::HandleCurveEditorKeySelectionChanged);
+	}
+
+	FCurveEditorSelection& CurveEditorSelection = CurveEditor->GetSelection();
+
+	CurveEditorSelection.Clear();
+
+	for (const TPair<FCurveModelID, TUniquePtr<FCurveModel>>& Pair : CurveEditor->GetCurves())
+	{
+		FCurveModel* const CurveModel = Pair.Value.Get();
+		if (!CurveModel)
+		{
+			continue;
+		}
+
+		const TSharedPtr<FChannelModel> ChannelModel = Utils::FindChannelModelForCurve(Timeline.AsShared(), *CurveModel);
+		if (!ChannelModel.IsValid())
+		{
+			continue;
+		}
+
+		for (const FSequencerSelectedKey& SelectedKey : SelectedKeys)
+		{
+			if (SelectedKey.WeakChannel.Pin() != ChannelModel)
+			{
+				continue;
+			}
+
+			CurveEditorSelection.Add(FCurvePointHandle(Pair.Key, ECurvePointType::Key, SelectedKey.KeyHandle));
+		}
+	}
+}
+
+#if 1==2
+void FToolableTimelineKeySelection::SyncRangeSelectionToCurveEditor(const TRange<FFrameNumber>& InRange)
+{
+	if (bSyncDisabled || !IsSimpleView())
+	{
+		return;
+	}
+	TGuardValue<bool> Guard(bSyncDisabled, true);
+
+	if (!FToolableTimelineBaseTool::IsValidToolRange(InRange))
+	{
+		return;
+	}
+
+	const TSharedPtr<ISequencer> Sequencer = Timeline.GetSequencer();
+	if (!Sequencer.IsValid())
+	{
+		return;
+	}
+
+	const TSharedPtr<FCurveEditor> CurveEditor = Utils::GetSequencerCurveEditor(*Sequencer);
+	if (!CurveEditor.IsValid())
+	{
+		return;
+	}
+
+	const TSharedRef<FToolableTimeSliderController> TimeSliderController = Timeline.GetTimeSliderController();
+	const FFrameRate TickResolution = TimeSliderController->GetTickResolution();
+
+	const double MinTime = InRange.GetLowerBoundValue() / TickResolution;
+	const double MaxTime = InRange.GetUpperBoundValue() / TickResolution;
+	const TRange<double> SecondsRange(MinTime, MaxTime);
+
+	FCurveEditorSelection& CurveEditorSelection = CurveEditor->GetSelection();
+	CurveEditorSelection.Clear();
+
+	for (const TPair<FCurveModelID, TUniquePtr<FCurveModel>>& Pair : CurveEditor->GetCurves())
+	{
+		FCurveModel* const CurveModel = Pair.Value.Get();
+		if (!CurveModel)
+		{
+			continue;
+		}
+
+		TArray<FKeyHandle> KeyHandles;
+		CurveModel->GetKeys(
+			SecondsRange.GetLowerBoundValue(),
+			SecondsRange.GetUpperBoundValue(),
+			TNumericLimits<double>::Lowest(),
+			TNumericLimits<double>::Max(),
+			KeyHandles
+		);
+
+		TArray<FKeyPosition> KeyPositions;
+		KeyPositions.SetNumUninitialized(KeyHandles.Num());
+		CurveModel->GetKeyPositions(KeyHandles, KeyPositions);
+
+		for (int32 KeyIndex = 0; KeyIndex < KeyHandles.Num(); ++KeyIndex)
+		{
+			if (SecondsRange.Contains(KeyPositions[KeyIndex].InputValue))
+			{
+				CurveEditorSelection.Add(FCurvePointHandle(Pair.Key, ECurvePointType::Key, KeyHandles[KeyIndex]));
+			}
+		}
+	}
+}
+#endif
+
+void FToolableTimelineKeySelection::SyncSelectionToSequencerAndCurveEditor()
+{
+	if (bSyncDisabled || !IsSimpleView())
+	{
+		return;
+	}
+	TGuardValue<bool> Guard(bSyncDisabled, true);
+
+	SyncKeySelectionToSequencer();
+	SyncKeySelectionToCurveEditor();
+}
+
+void FToolableTimelineKeySelection::ClearCurveEditorSelection() const
+{
+	const TSharedPtr<ISequencer> Sequencer = Timeline.GetSequencer();
+	if (!Sequencer.IsValid())
+	{
+		return;
+	}
+
+	const TSharedPtr<FCurveEditor> CurveEditor = Utils::GetSequencerCurveEditor(*Sequencer);
+	if (!CurveEditor.IsValid())
+	{
+		return;
+	}
+
+	CurveEditor->GetSelection().Clear();
+}
+
+bool FToolableTimelineKeySelection::IsSimpleView() const
+{
+	const TSharedPtr<FSequencer> Sequencer = Timeline.GetSequencer();
+	return Sequencer.IsValid() && Sequencer->IsSimpleView();
+}
+
+} // namespace UE::Sequencer::ToolableTimeline
diff --git a/Engine/Source/Editor/Sequencer/Private/Tools/SequencerEditTool_Movement.cpp b/Engine/Source/Editor/Sequencer/Private/Tools/SequencerEditTool_Movement.cpp
--- a/Engine/Source/Editor/Sequencer/Private/Tools/SequencerEditTool_Movement.cpp
+++ b/Engine/Source/Editor/Sequencer/Private/Tools/SequencerEditTool_Movement.cpp
@@ -534,7 +534,7 @@ void FSequencerEditTool_Movement::UpdateCursorDecorator(const FGeometry& MyGeome
 	}
 }
 
-FCursorReply FSequencerEditTool_Movement::OnCursorQuery(const FGeometry& MyGeometry, const FPointerEvent& CursorEvent) const
+FCursorReply FSequencerEditTool_Movement::OnCursorQuery(const SWidget& OwnerWidget, const FGeometry& MyGeometry, const FPointerEvent& CursorEvent) const
 {
 	using namespace UE::Sequencer;
 
diff --git a/Engine/Source/Editor/Sequencer/Private/Tools/SequencerEditTool_Movement.h b/Engine/Source/Editor/Sequencer/Private/Tools/SequencerEditTool_Movement.h
--- a/Engine/Source/Editor/Sequencer/Private/Tools/SequencerEditTool_Movement.h
+++ b/Engine/Source/Editor/Sequencer/Private/Tools/SequencerEditTool_Movement.h
@@ -31,6 +31,8 @@ public:
 
 public:
 
+	using FSequencerEditTool::OnCursorQuery;
+
 	// ISequencerEditTool interface
 
 	virtual FReply OnMouseButtonDown(SWidget& OwnerWidget, const FGeometry& MyGeometry, const FPointerEvent& MouseEvent) override;
@@ -38,7 +40,7 @@ public:
 	virtual FReply OnMouseMove(SWidget& OwnerWidget, const FGeometry& MyGeometry, const FPointerEvent& MouseEvent) override;
 	virtual void OnMouseCaptureLost() override;
 	virtual int32 OnPaint(const FGeometry& AllottedGeometry, const FSlateRect& MyCullingRect, FSlateWindowElementList& OutDrawElements, int32 LayerId) const override;
-	virtual FCursorReply OnCursorQuery(const FGeometry& MyGeometry, const FPointerEvent& CursorEvent) const override;
+	virtual FCursorReply OnCursorQuery(const SWidget& OwnerWidget, const FGeometry& MyGeometry, const FPointerEvent& CursorEvent) const override;
 	virtual FName GetIdentifier() const override;
 	virtual bool CanDeactivate() const override;
 	virtual TSharedPtr<UE::Sequencer::ITrackAreaHotspot> GetDragHotspot() const override;
diff --git a/Engine/Source/Editor/Sequencer/Private/Tools/SequencerEditTool_Selection.cpp b/Engine/Source/Editor/Sequencer/Private/Tools/SequencerEditTool_Selection.cpp
--- a/Engine/Source/Editor/Sequencer/Private/Tools/SequencerEditTool_Selection.cpp
+++ b/Engine/Source/Editor/Sequencer/Private/Tools/SequencerEditTool_Selection.cpp
@@ -432,7 +432,7 @@ FSequencerEditTool_Selection::FSequencerEditTool_Selection(FSequencer& InSequenc
 { }
 
 
-FCursorReply FSequencerEditTool_Selection::OnCursorQuery(const FGeometry& MyGeometry, const FPointerEvent& CursorEvent) const
+FCursorReply FSequencerEditTool_Selection::OnCursorQuery(const SWidget& OwnerWidget, const FGeometry& MyGeometry, const FPointerEvent& CursorEvent) const
 {
 	if (DragOperation)
 	{
diff --git a/Engine/Source/Editor/Sequencer/Private/Tools/SequencerEditTool_Selection.h b/Engine/Source/Editor/Sequencer/Private/Tools/SequencerEditTool_Selection.h
--- a/Engine/Source/Editor/Sequencer/Private/Tools/SequencerEditTool_Selection.h
+++ b/Engine/Source/Editor/Sequencer/Private/Tools/SequencerEditTool_Selection.h
@@ -31,6 +31,8 @@ public:
 
 public:
 
+	using FSequencerEditTool::OnCursorQuery;
+
 	// ISequencerEditTool interface	
 	virtual int32 OnPaint(const FGeometry& AllottedGeometry, const FSlateRect& MyCullingRect, FSlateWindowElementList& OutDrawElements, int32 LayerId) const override;
 	virtual FReply OnMouseButtonDown(SWidget& OwnerWidget, const FGeometry& MyGeometry, const FPointerEvent& MouseEvent) override;
@@ -38,7 +40,7 @@ public:
 	virtual FReply OnMouseMove(SWidget& OwnerWidget, const FGeometry& MyGeometry, const FPointerEvent& MouseEvent) override;
 	virtual void OnMouseLeave(SWidget& OwnerWidget, const FPointerEvent& MouseEvent) override;
 	virtual void OnMouseCaptureLost() override;
-	virtual FCursorReply OnCursorQuery(const FGeometry& MyGeometry, const FPointerEvent& CursorEvent) const override;
+	virtual FCursorReply OnCursorQuery(const SWidget& OwnerWidget, const FGeometry& MyGeometry, const FPointerEvent& CursorEvent) const override;
 	virtual FName GetIdentifier() const override;
 	virtual bool CanDeactivate() const override;
 	virtual TSharedPtr<UE::Sequencer::ITrackAreaHotspot> GetDragHotspot() const override;
diff --git a/Engine/Source/Editor/Sequencer/Public/ISequencer.h b/Engine/Source/Editor/Sequencer/Public/ISequencer.h
--- a/Engine/Source/Editor/Sequencer/Public/ISequencer.h
+++ b/Engine/Source/Editor/Sequencer/Public/ISequencer.h
@@ -1014,6 +1014,12 @@ public:
 
 	virtual TSharedPtr<UE::Sequencer::SOutlinerView> GetOutlinerViewWidget() const = 0;
 
+	/** @return True if the Sequencer is in the simplified view */
+	virtual bool IsSimpleView() const = 0;
+
+	/** Enables or disables the Sequencer simplified view */
+	virtual void EnableSimpleView(const bool bInEnable) = 0;
+
 protected:
 	FOnInitializeDetailsPanel InitializeDetailsPanelEvent;
 	FOnCameraAddedToSequencer CameraAddedToSequencer;
diff --git a/Engine/Source/Editor/Sequencer/Public/ISequencerModule.h b/Engine/Source/Editor/Sequencer/Public/ISequencerModule.h
--- a/Engine/Source/Editor/Sequencer/Public/ISequencerModule.h
+++ b/Engine/Source/Editor/Sequencer/Public/ISequencerModule.h
@@ -163,6 +163,9 @@ struct FSequencerHostCapabilities
 	/** Do we support drag and drop operation */
 	bool bSupportsDragAndDrop;
 
+	/** Do we support toggling between a minimalized, simple view instead of the full Sequencer view */
+	bool bSupportsSimpleView;
+
 	FSequencerHostCapabilities()
 		: bSupportsSaveMovieSceneAsset(false)
 		, bSupportsCurveEditor(false)
@@ -172,6 +175,7 @@ struct FSequencerHostCapabilities
 		, bSupportsSidebar(false)
 		, bSupportsViewportSelectability(false)
 		, bSupportsDragAndDrop(true)
+		, bSupportsSimpleView(true)
 	{}
 };
 
diff --git a/Engine/Source/Editor/Sequencer/Public/SequencerCommands.h b/Engine/Source/Editor/Sequencer/Public/SequencerCommands.h
--- a/Engine/Source/Editor/Sequencer/Public/SequencerCommands.h
+++ b/Engine/Source/Editor/Sequencer/Public/SequencerCommands.h
@@ -531,6 +531,18 @@ public:
 	/** Toggle the soloed state of Sequencer track selection */
 	TSharedPtr< FUICommandInfo > ToggleTrackSelectionSolo;
 
+	/** Toggle between the toolbar panel being visible and hidden */
+	TSharedPtr< FUICommandInfo > ToggleToolBarVisible;
+
+	/** Toggle between the transport controls panel being visible and hidden */
+	TSharedPtr< FUICommandInfo > ToggleTransportPanelVisible;
+
+	/** Toggle between the transport controls panel being popped out in its own tab window and docked within its inline content */
+	TSharedPtr< FUICommandInfo > ToggleTransportPanelDocked;
+
+	/** Toggle between Sequencer simplified view and full views */
+	TSharedPtr< FUICommandInfo > ToggleSimpleView;
+
 	/**
 	 * Initialize commands
 	 */
diff --git a/Engine/Source/Editor/Sequencer/Public/SequencerSettings.h b/Engine/Source/Editor/Sequencer/Public/SequencerSettings.h
--- a/Engine/Source/Editor/Sequencer/Public/SequencerSettings.h
+++ b/Engine/Source/Editor/Sequencer/Public/SequencerSettings.h
@@ -637,6 +637,11 @@ public:
 	/** Sets the popped out state of the transport controls panel in the Sequencer tab */
 	UE_API void SetTransportControlsPoppedOut(const bool bInPoppedOut);
 
+	/** @return True if the last the time Sequencer was open, it was in the simplified view and not the full view */
+	bool IsSimpleView() const { return bSimpleView; }
+	/** Switches between full Sequencer and a simplified Sequencer view */
+	void SetIsSimpleView(const bool bInSimpleView);
+
 protected:
 
 	/** The auto change mode (auto-key, auto-track or none). */
@@ -1029,6 +1034,10 @@ protected:
 	UPROPERTY(config, EditAnywhere, Category = General)
 	bool bTransportControlsPoppedOut;
 
+	/** Last saved state of the Sequencer view: full or simple */
+	UPROPERTY(config)
+	bool bSimpleView;
+
 	UFUNCTION()
 	static UE_API bool ShouldShowThumbnailCaptureSettings();
 
diff --git a/Engine/Source/Editor/Sequencer/Public/SequencerTimeSliderController.h b/Engine/Source/Editor/Sequencer/Public/SequencerTimeSliderController.h
--- a/Engine/Source/Editor/Sequencer/Public/SequencerTimeSliderController.h
+++ b/Engine/Source/Editor/Sequencer/Public/SequencerTimeSliderController.h
@@ -12,10 +12,13 @@
 
 class FSlateWindowElementList;
 struct FContextMenuSuppressor;
+struct FGeometry;
+struct FPointerEvent;
 struct FSlateBrush;
 class FSlateFontMeasure;
 class FSequencer;
 class IPropertyTypeCustomization;
+class USequencerSettings;
 
 /**
  * A time slider controller for sequencer
@@ -38,18 +41,23 @@ public:
 	*/
 	float DetermineOptimalSpacing(float InPixelsPerInput, uint32 MinTick, float MinTickSpacing) const;
 
+	virtual const FTimeSliderArgs& GetTimeSliderArgs() const { return TimeSliderArgs; }
+
 	/** ITimeSliderController Interface */
 	virtual int32 OnPaintTimeSlider( bool bMirrorLabels, const FGeometry& AllottedGeometry, const FSlateRect& MyCullingRect, FSlateWindowElementList& OutDrawElements, int32 LayerId, const FWidgetStyle& InWidgetStyle, bool bParentEnabled ) const override;
-	virtual FReply OnMouseButtonDown( SWidget& WidgetOwner, const FGeometry& MyGeometry, const FPointerEvent& MouseEvent ) override;
-	virtual FReply OnMouseButtonUp( SWidget& WidgetOwner, const FGeometry& MyGeometry, const FPointerEvent& MouseEvent ) override;
-	virtual FReply OnMouseButtonDoubleClick(SWidget& WidgetOwner, const FGeometry& MyGeometry, const FPointerEvent& MouseEvent) override;
-	virtual FReply OnMouseMove( SWidget& WidgetOwner, const FGeometry& MyGeometry, const FPointerEvent& MouseEvent ) override;
-	virtual FReply OnMouseWheel( SWidget& WidgetOwner, const FGeometry& MyGeometry, const FPointerEvent& MouseEvent ) override;
-	virtual FReply OnTimeSliderMouseMove( SWidget& OwnerWidget, const FGeometry& MyGeometry, const FPointerEvent& MouseEvent ) override;
-	virtual FCursorReply OnCursorQuery( TSharedRef<const SWidget> WidgetOwner, const FGeometry& MyGeometry, const FPointerEvent& CursorEvent ) const override;
+	virtual FReply OnMouseButtonDown(SWidget& OwnerWidget, const FGeometry& InGeometry, const FPointerEvent& InPointEvent) override;
+	virtual FReply OnMouseButtonUp(SWidget& OwnerWidget, const FGeometry& InGeometry, const FPointerEvent& InPointerEvent) override;
+	virtual FReply OnMouseButtonDoubleClick(SWidget& OwnerWidget, const FGeometry& InGeometry, const FPointerEvent& InPointerEvent) override;
+	virtual FReply OnMouseMove(SWidget& OwnerWidget, const FGeometry& InGeometry, const FPointerEvent& InPointerEvent) override;
+	virtual FReply OnMouseWheel(SWidget& OwnerWidget, const FGeometry& InGeometry, const FPointerEvent& InPointerEvent) override;
+	virtual FReply OnTimeSliderMouseMove(SWidget& OwnerWidget, const FGeometry& InGeometry, const FPointerEvent& InPointerEvent) override;
+	virtual FCursorReply OnCursorQuery( const SWidget& OwnerWidget, const FGeometry& MyGeometry, const FPointerEvent& CursorEvent ) const override;
 	virtual double ComputeHeight() const override;
 	/** End ITimeSliderController Interface */
 
+	/** Get the owning sequencer instance */
+	TSharedPtr<ISequencer> GetSequencer() const;
+
 	/** Get the current play rate for this controller */
 	virtual FFrameRate GetDisplayRate() const override { return TimeSliderArgs.DisplayRate.Get(); }
 
@@ -146,7 +154,14 @@ public:
 	 */
 	virtual int32 OnPaintViewArea( const FGeometry& AllottedGeometry, const FSlateRect& MyCullingRect, FSlateWindowElementList& OutDrawElements, int32 LayerId, bool bEnabled, const FPaintViewAreaArgs& Args ) const override;
 
-
+	/**
+	 * Call this method when the user's interaction has changed the scrub position
+	 *
+	 * @param NewValue				Value resulting from the user's interaction
+	 * @param bIsScrubbing			True if done via scrubbing, false if just releasing scrubbing
+	 * @param bEvaluate				If true evaluate, if not just change time
+	 */
+	void CommitScrubPosition(FFrameTime NewValue, bool bIsScrubbing, bool bEvaluate);
 
 public:
 
@@ -166,7 +181,7 @@ public:
 	struct FScrubRangeToScreen
 	{
 		double ViewStart;
-		float  PixelsPerInput;
+		float PixelsPerInput;
 
 		FScrubRangeToScreen(const TRange<double>& InViewInput, const FVector2D& InWidgetSize)
 		{
@@ -203,6 +218,8 @@ public:
 		FSlateRect CullingRect;
 		/** Color of each tick */
 		FLinearColor TickColor;
+		/** Color of each tick's text, if displayed */
+		FLinearColor TickTextColor;
 		/** Offset in Y where to start the tick */
 		float TickOffset;
 		/** Height in of major ticks */
@@ -218,26 +235,40 @@ public:
 		
 	};
 
+	bool IsEvaluating() const { return bIsEvaluating; }
+
 	/** Set that's evaluating */
 	void SetIsEvaluating()
 	{
 		bIsEvaluating = true;
 	}
 
-private:
+protected:
 
-	FReply OnMouseMoveImpl( SWidget& WidgetOwner, const FGeometry& MyGeometry, const FPointerEvent& MouseEvent, bool bFromTimeSlider );
+	virtual FReply OnMouseMoveImpl(SWidget& WidgetOwner
+		, const FGeometry& InGeometry
+		, const FPointerEvent& InPointerEvent
+		, const bool bInFromTimeSlider);
 
-private:
+	virtual FReply OnMouseButtonUp_ContextMenu(SWidget& WidgetOwner, const FGeometry& InGeometry, const FPointerEvent& InPointerEvent);
+	virtual FReply OnMouseButtonUp_MouseDragType(SWidget& WidgetOwner
+		, const FGeometry& InGeometry
+		, const FPointerEvent& InPointerEvent
+		, const FScrubRangeToScreen& InRangeToScreen
+		, const FFrameTime& InMouseTime);
 
-	/**
-	 * Call this method when the user's interaction has changed the scrub position
-	 *
-	 * @param NewValue				Value resulting from the user's interaction
-	 * @param bIsScrubbing			True if done via scrubbing, false if just releasing scrubbing
-	 * @param bEvaluate				If true evaluate, if not just change time
-	 */
-	void CommitScrubPosition( FFrameTime NewValue, bool bIsScrubbing, bool bEvaluate);
+	virtual FReply OnMouseMove_LeftAndMiddleMouseDown(const FGeometry& InGeometry
+		, const FPointerEvent& InPointerEvent
+		, const bool bInFromTimeSlider
+		, const FScrubRangeToScreen& InRangeToScreen
+		, const USequencerSettings& InSequencerSettings);
+	virtual FReply OnMouseMove_RightMouseDown(const FGeometry& InGeometry
+		, const FPointerEvent& InPointerEvent, const USequencerSettings& InSequencerSettings);
+
+	virtual void HandleDragMark_NoDragType(const int32 InMarkIndex);
+	virtual void HandleDragMark(const FFrameNumber& InLastDiffFrame);
+
+	void ResetMouseDown(const FGeometry& InGeometry, const FPointerEvent& InPointerEvent);
 
 	/**
 	 * Draw time tick marks
@@ -262,21 +293,31 @@ private:
 	/**
 	 * Draw horizontal separator lines between tracks
 	 */
-	void DrawHorizontalTrackSeparatorLines(FSlateWindowElementList& OutDrawElements, const TRange<double>& InViewRange, const FScrubRangeToScreen& InRangeToScreen, FDrawTickArgs& InArgs) const;
+	void DrawHorizontalTrackSeparatorLines(FSlateWindowElementList& OutDrawElements, FDrawTickArgs& InArgs) const;
 
 	/**
 	 * Draw the selection range.
 	 *
 	 * @return The new layer ID.
 	 */
-	int32 DrawSelectionRange(const FGeometry& AllottedGeometry, const FSlateRect& MyCullingRect, FSlateWindowElementList& OutDrawElements, int32 LayerId, const FScrubRangeToScreen& RangeToScreen, const FPaintPlaybackRangeArgs& Args) const;
+	int32 DrawSelectionRange(const FGeometry& InGeometry,
+		const FSlateRect& InCullingRect,
+		FSlateWindowElementList& OutDrawElements,
+		int32 LayerId,
+		const FScrubRangeToScreen& InRangeToScreen,
+		const FPaintPlaybackRangeArgs& InRangeArgs) const;
 
 	/**
 	 * Draw the playback range.
 	 *
 	 * @return the new layer ID
 	 */
-	int32 DrawPlaybackRange(const FGeometry& AllottedGeometry, const FSlateRect& MyCullingRect, FSlateWindowElementList& OutDrawElements, int32 LayerId, const FScrubRangeToScreen& RangeToScreen, const FPaintPlaybackRangeArgs& Args) const;
+	int32 DrawPlaybackRange(const FGeometry& InGeometry,
+		const FSlateRect& InCullingRect,
+		FSlateWindowElementList& OutDrawElements,
+		int32 LayerId,
+		const FScrubRangeToScreen& InRangeToScreen,
+		const FPaintPlaybackRangeArgs& InRangeArgs) const;
 
 	/**
 	 * Draw the playback range.
@@ -297,8 +338,12 @@ private:
 	 *
 	 * @return the new layer ID
 	 */
-	int32 DrawMarkedFrames(const FGeometry& AllottedGeometry, const FScrubRangeToScreen& RangeToScreen, FSlateWindowElementList& OutDrawElements, int32 LayerId, const ESlateDrawEffect& DrawEffects, const FWidgetStyle& InWidgetStyle, bool bDrawLabels) const;
-
+	int32 DrawMarkedFrames(const FGeometry& InGeometry,
+		const FScrubRangeToScreen& InRangeToScreen,
+		FSlateWindowElementList& OutDrawElements,
+		int32 LayerId,
+		const ESlateDrawEffect& InDrawEffects,
+		bool bInDrawLabels) const;
 
 	/**
 	 * Draw any scaling anchors.
@@ -307,8 +352,6 @@ private:
 	 */
 	int32 DrawScalingAnchors( const FGeometry& AllottedGeometry, const FScrubRangeToScreen& RangeToScreen, FSlateWindowElementList& OutDrawElements, int32 LayerId, const ESlateDrawEffect& DrawEffects ) const;
 
-private:
-
 	/**
 	 * Hit test the lower bound of a range
 	 */
@@ -319,6 +362,8 @@ private:
 	 */
 	bool HitTestRangeEnd(const FScrubRangeToScreen& RangeToScreen, const TRange<double>& Range, float HitPixel) const;
 
+public:
+
 	/**
 	 * Hit test marks
 	 *
@@ -326,6 +371,8 @@ private:
 	 */
 	bool HitTestMark(const FGeometry& AllottedGeometry, const FScrubRangeToScreen& RangeToScreen, float HitPixel, bool bTestLabelBox, int32* OutMarkIndex = nullptr, FFrameNumber* OutMarkFrameNumber = nullptr) const;
 
+protected:
+
 	/**
 	 * Get marked frame label box size
 	 */
@@ -333,24 +380,16 @@ private:
 
 	FFrameTime SnapTimeToNearestKey(const FPointerEvent& MouseEvent, const FScrubRangeToScreen& RangeToScreen, float CursorPos, FFrameTime InTime) const;
 
+public:
 	void SetPlaybackRangeStart(FFrameNumber NewStart);
 	void SetPlaybackRangeEnd(FFrameNumber NewEnd);
 
 	void SetSelectionRangeStart(FFrameNumber NewStart);
 	void SetSelectionRangeEnd(FFrameNumber NewEnd);
 
-	TSharedRef<SWidget> OpenSetPlaybackRangeMenu(const FGeometry& MyGeometry, const FPointerEvent& MouseEvent);
-	FFrameTime SnapSequencerTime(FFrameTime InTime) const;
 	FFrameTime ComputeScrubTimeFromMouse(const FGeometry& Geometry, const FPointerEvent& MouseEvent, FScrubRangeToScreen RangeToScreen) const;
 	FFrameTime ComputeFrameTimeFromMouse(const FGeometry& Geometry, FVector2D ScreenSpacePosition, FScrubRangeToScreen RangeToScreen, bool CheckSnapping = true) const;
 
-	void HandleMarkSelection(int32 InMarkIndex);
-	void AddMarkAtFrame(FFrameNumber FrameNumber);
-	void DeleteMarkAtIndex(int32 InMarkIndex);
-	void DeleteAllMarks();
-
-private:
-
 	/**
 	 * Get the pixel matrics of the Scrubber
 	 * @param ScrubTime			The qualified time of the scrubber
@@ -358,11 +397,23 @@ private:
 	 * @param DilationPixels	Number of pixels to dilate the handle by
 	 * return FScrubberMetrics struct
 	 */
-	FScrubberMetrics GetScrubPixelMetrics(const FQualifiedFrameTime& ScrubTime, const FScrubRangeToScreen& RangeToScreen, float DilationPixels = 0.f) const;
+	virtual FScrubberMetrics GetScrubPixelMetrics(const FQualifiedFrameTime& ScrubTime, const FScrubRangeToScreen& RangeToScreen, float DilationPixels = 0.f) const;
+
+	void ClearMarkSelection();
+	TOptional<FFrameNumber> GetMarkFrameNumber(const int32 InMarkIndex) const;
+	void AddMarkAtFrame(const FFrameNumber InFrameNumber);
+	void SetMarkAtFrame(const int32 InMarkIndex, const FFrameNumber InFrameNumber);
+	void DeleteMarkAtIndex(const int32 InMarkIndex);
+	void DeleteAllMarks();
+
+protected:
+	virtual TSharedPtr<SWidget> OpenContextMenu(const FGeometry& InGeometry, const FPointerEvent& InPointerEvent);
+
+	FFrameTime SnapSequencerTime(FFrameTime InTime) const;
 
 	FScrubberMetrics GetHitTestScrubPixelMetrics(const FScrubRangeToScreen& RangeToScreen) const;
 
-private:
+	void ResetMouseInput();
 
 	/** Pointer back to the sequencer object */
 	TWeakPtr<FSequencer> WeakSequencer;
@@ -384,21 +435,20 @@ private:
 	
 	/** Total mouse delta during dragging **/
 	float DistanceDragged;
-	
-	/** If we are dragging a scrubber or dragging to set the time range */
-	enum DragType
+
+	/** Drag operations for this time controller */
+	enum EDragType
 	{
+		DRAG_NONE,
 		DRAG_SCRUBBING_TIME,
 		DRAG_SETTING_RANGE,
 		DRAG_PLAYBACK_START,
 		DRAG_PLAYBACK_END,
 		DRAG_SELECTION_START,
 		DRAG_SELECTION_END,
-		DRAG_MARK,
-		DRAG_NONE
+		DRAG_MARK
 	};
-	
-	DragType MouseDragType;
+	EDragType MouseDragType;
 	
 	/** If mouse down was in time scrubbing region, only allow setting time when mouse is pressed down in the region */
 	bool bMouseDownInRegion;
@@ -429,12 +479,11 @@ private:
 
 	/** When > 0, we should not show context menus */
 	int32 ContextMenuSuppression;
-	
-	/** If evaluating, if not we draw the time box to be yellow not default*/
+
+	/** If not evaluating, we draw the time box to be yellow and not default */
 	bool bIsEvaluating = true;
 
 	friend FContextMenuSuppressor;
-	
 };
 
 struct FContextMenuSuppressor
diff --git a/Engine/Source/Editor/Sequencer/Sequencer.Build.cs b/Engine/Source/Editor/Sequencer/Sequencer.Build.cs
--- a/Engine/Source/Editor/Sequencer/Sequencer.Build.cs
+++ b/Engine/Source/Editor/Sequencer/Sequencer.Build.cs
@@ -25,6 +25,7 @@ public class Sequencer : ModuleRules
 				"Core", 
 				"CoreUObject", 
 				"CurveEditor",
+				"DeveloperSettings",
 				"InputCore",
 				"Engine", 
 				"EditorInteractiveToolsFramework",
@@ -56,7 +57,8 @@ public class Sequencer : ModuleRules
 				"UniversalObjectLocator",
 				"UniversalObjectLocatorEditor",
 				"RenderCore",
-				}
+				"WorkspaceMenuStructure"
+			}
 			);
 
 		if (Target.bBuildEditor)
diff --git a/Engine/Source/Editor/SequencerCore/Private/MVVM/Views/KeyRenderer.cpp b/Engine/Source/Editor/SequencerCore/Private/MVVM/Views/KeyRenderer.cpp
--- a/Engine/Source/Editor/SequencerCore/Private/MVVM/Views/KeyRenderer.cpp
+++ b/Engine/Source/Editor/SequencerCore/Private/MVVM/Views/KeyRenderer.cpp
@@ -21,7 +21,12 @@ namespace UE::Sequencer
 
 void FKeyRenderer::Initialize(const FViewModelPtr& InViewModel)
 {
-	WeakViewModel = InViewModel;
+	WeakViewModels = { InViewModel };
+}
+
+void FKeyRenderer::Initialize(const TSet<FWeakViewModelPtr>& InViewModels)
+{
+	WeakViewModels = InViewModels;
 }
 
 bool FKeyRenderer::HasCurve() const
@@ -29,69 +34,148 @@ bool FKeyRenderer::HasCurve() const
 	return PrecomputedCurve.Num() >= 2;
 }
 
-void FKeyRenderer::HitTestKeys(const FFrameTime& Time, TArray<FKeysForModel>& OutAllKeys) const
+void FKeyRenderer::HitTestKeys(const FFrameTime& Time, TArray<FKeysForModel>& OutAllKeys, const bool bInLeftAlignCollision) const
 {
-	const FFrameTime HalfKeyWidth = KeyWidthInFrames * .5f;
-
-	// Hit test the actual structure used for drawing keys to ensure we are 100% accurate
-	// PrecomputedKeys is an array sorted by time, so we can binary search it to find the hit key
-	const int32 LowerIndex = Algo::LowerBoundBy(PrecomputedKeys, Time - HalfKeyWidth, &FKey::FinalKeyPosition);
-	const int32 UpperIndex = Algo::UpperBoundBy(PrecomputedKeys, Time + HalfKeyWidth, &FKey::FinalKeyPosition);
-
-	if (LowerIndex >= PrecomputedKeys.Num())
+	const int32 PrecomputedKeysCount = PrecomputedKeys.Num();
+	if (PrecomputedKeysCount == 0)
 	{
-		// No hit - we ran off the end of the drawn keys
 		return;
 	}
 
-	// We may have hit many overlapping keys, so we need to choose the one closest to the time
-	// Assume the first one is correct, then keep checking the distance between that and the next
-	// Break as soon as we find a key that is further away than the last (we know we're moving away from the time at that point)
-	int32 HitIndex = LowerIndex;
-	for (int32 Index = LowerIndex+1; Index < UpperIndex && Index < PrecomputedKeys.Num(); ++Index)
+	const FFrameTime HalfKeyWidth = KeyWidthInFrames * .5f;
+
+	int32 HitIndex = INDEX_NONE;
+
+	if (bInLeftAlignCollision)
 	{
-		const FKey& PrevKey = PrecomputedKeys[Index-1];
-		const FKey& ThisKey = PrecomputedKeys[Index];
+		// Left-aligned collision: key occupies [FinalKeyPosition, FinalKeyPosition + KeyWidthInFrames]
+		// Keep it fast:
+		// - candidates are keys whose start is within one key width behind Time, up to Time
+		// - we early-break once Start > Time and we're already getting worse (Start only increases)
+		const FFrameTime KeyWidth = KeyWidthInFrames;
 
-		// Find the distance from the middle of the key to the hit time
-		FFrameTime PrevKeyDistance = FMath::Abs((PrevKey.KeyTickStart + PrevKey.KeyTickEnd) * .5f - Time);
-		FFrameTime ThisKeyDistance = FMath::Abs((ThisKey.KeyTickStart + ThisKey.KeyTickEnd) * .5f - Time);
+		const int32 LowerIndex = Algo::LowerBoundBy(PrecomputedKeys, Time - KeyWidth, &FKey::FinalKeyPosition);
+		const int32 UpperIndex = Algo::UpperBoundBy(PrecomputedKeys, Time,            &FKey::FinalKeyPosition);
 
-		if (ThisKeyDistance <= PrevKeyDistance)
+		if (LowerIndex >= PrecomputedKeys.Num())
 		{
-			HitIndex = Index;
+			return;
 		}
-		else
+
+		FFrameTime BestDistance = TNumericLimits<FFrameNumber>::Max();
+		for (int32 Index = LowerIndex; Index < UpperIndex && Index < PrecomputedKeys.Num(); ++Index)
 		{
-			break;
+			const FKey& Key = PrecomputedKeys[Index];
+
+			const FFrameTime Start = Key.FinalKeyPosition;
+			const FFrameTime End   = Start + KeyWidth;
+
+			FFrameTime Distance = FFrameTime(0);
+
+			if (Time < Start)
+			{
+				Distance = Start - Time;
+			}
+			else if (Time > End)
+			{
+				Distance = Time - End;
+			}
+
+			if (Distance <= BestDistance)
+			{
+				BestDistance = Distance;
+				HitIndex = Index;
+
+				if (Distance == FFrameTime(0))
+				{
+					// Can't do better than a direct hit
+					break;
+				}
+			}
+			else if (Start > Time)
+			{
+				// From here on, Start only increases, so Distance = Start-Time only increases.
+				break;
+			}
+		}
+
+		if (HitIndex == INDEX_NONE)
+		{
+			return;
 		}
-	}
 
-	if (HitIndex < PrecomputedKeys.Num() && FMath::Abs(PrecomputedKeys[HitIndex].FinalKeyPosition - Time) <= HalfKeyWidth)
+		const FFrameTime HitStart = PrecomputedKeys[HitIndex].FinalKeyPosition;
+		const FFrameTime HitEnd   = HitStart + KeyWidth;
+		if (Time < HitStart || Time > HitEnd)
+		{
+			return;
+		}
+	}
+	else // Center-aligned collision
 	{
-		FFrameTime Start = PrecomputedKeys[HitIndex].KeyTickStart;
-		FFrameTime End   = PrecomputedKeys[HitIndex].KeyTickEnd;
+		// Hit test the actual structure used for drawing keys to ensure we are 100% accurate
+		// PrecomputedKeys is an array sorted by time, so we can binary search it to find the hit key
+		const int32 LowerIndex = Algo::LowerBoundBy(PrecomputedKeys, Time - HalfKeyWidth, &FKey::FinalKeyPosition);
+		const int32 UpperIndex = Algo::UpperBoundBy(PrecomputedKeys, Time + HalfKeyWidth, &FKey::FinalKeyPosition);
+		
+		if (LowerIndex >= PrecomputedKeysCount)
+		{
+			// No hit - we ran off the end of the drawn keys
+			return;
+		}
 
-		// We have a hit. Gather anything from the key draw info that overlaps the hit range
-		for (const FCachedKeyDrawInformation& Entry : KeyDrawInfo)
+		// We may have hit many overlapping keys, so we need to choose the one closest to the time
+		// Assume the first one is correct, then keep checking the distance between that and the next
+		// Break as soon as we find a key that is further away than the last (we know we're moving away from the time at that point)
+		HitIndex = LowerIndex;
+		for (int32 Index = LowerIndex + 1; Index < UpperIndex && Index < PrecomputedKeysCount; ++Index)
 		{
-			const int32 FrameStart = Algo::LowerBound(Entry.FramesInRange, Start);
-			const int32 FrameEnd   = Algo::UpperBound(Entry.FramesInRange, End);
-			const int32 NumKeys    = FrameEnd - FrameStart;
+			const FKey& PrevKey = PrecomputedKeys[Index - 1];
+			const FKey& ThisKey = PrecomputedKeys[Index];
+
+			// Find the distance from the middle of the key to the hit time
+			FFrameTime PrevKeyDistance = FMath::Abs((PrevKey.KeyTickStart + PrevKey.KeyTickEnd) * .5f - Time);
+			FFrameTime ThisKeyDistance = FMath::Abs((ThisKey.KeyTickStart + ThisKey.KeyTickEnd) * .5f - Time);
 
-			if (NumKeys > 0 && FrameStart < Entry.HandlesInRange.Num())
+			if (ThisKeyDistance <= PrevKeyDistance)
 			{
-				FKeysForModel Keys;
-				Keys.Model = Entry.WeakKeyExtension.Pin();
-				Keys.Keys.SetNumUninitialized(NumKeys);
+				HitIndex = Index;
+			}
+			else
+			{
+				break;
+			}
+		}
 
-				for (int32 Index = 0; Index < NumKeys; ++Index)
-				{
-					Keys.Keys[Index] = Entry.HandlesInRange[FrameStart + Index];
-				}
+		if (!(HitIndex < PrecomputedKeys.Num()
+			&& FMath::Abs(PrecomputedKeys[HitIndex].FinalKeyPosition - Time) <= HalfKeyWidth))
+		{
+			return;
+		}
+	}
 
-				OutAllKeys.Emplace(MoveTemp(Keys));
+	const FFrameTime Start = PrecomputedKeys[HitIndex].KeyTickStart;
+	const FFrameTime End   = PrecomputedKeys[HitIndex].KeyTickEnd;
+
+	// We have a hit. Gather anything from the key draw info that overlaps the hit range
+	for (const FCachedKeyDrawInformation& Entry : KeyDrawInfo)
+	{
+		const int32 FrameStart = Algo::LowerBound(Entry.FramesInRange, Start);
+		const int32 FrameEnd   = Algo::UpperBound(Entry.FramesInRange, End);
+		const int32 NumKeys    = FrameEnd - FrameStart;
+
+		if (NumKeys > 0 && FrameStart < Entry.HandlesInRange.Num())
+		{
+			FKeysForModel Keys;
+			Keys.Model = Entry.WeakKeyExtension.Pin();
+			Keys.Keys.SetNumUninitialized(NumKeys);
+
+			for (int32 Index = 0; Index < NumKeys; ++Index)
+			{
+				Keys.Keys[Index] = Entry.HandlesInRange[FrameStart + Index];
 			}
+
+			OutAllKeys.Emplace(MoveTemp(Keys));
 		}
 	}
 }
@@ -204,34 +288,37 @@ void FKeyRenderer::CacheKeyExtensions(const FKeyBatchParameters& Params) const
 {
 	KeyDrawInfo.Empty();
 
-	FViewModelPtr ViewModel = WeakViewModel.Pin();
-	if (!ViewModel)
+	// Walk through the current models and accumulate key renderer extensions for any collapsed children
+	for (const FWeakViewModelPtr& WeakViewModel : WeakViewModels)
 	{
-		return;
-	}
-
-	// Walk through the current model and accumulate key renderer extension for any collapsed children
-	constexpr bool bIncludeThis = true;
-	for (FParentFirstChildIterator ChildIt(ViewModel, bIncludeThis) ; ChildIt; ++ChildIt)
-	{
-		FViewModelPtr CurrentViewModel = *ChildIt;
-
-		TViewModelPtr<FLinkedOutlinerExtension> OutlinerExtension = CurrentViewModel.ImplicitCast();
-		if (OutlinerExtension && OutlinerExtension->GetLinkedOutlinerItem() && OutlinerExtension->GetLinkedOutlinerItem()->IsFilteredOut())
+		const TViewModelPtr<FViewModel> ViewModel = WeakViewModel.Pin();
+		if (!ViewModel.IsValid())
 		{
-			ChildIt.IgnoreCurrentChildren();
 			continue;
 		}
 
-		TViewModelPtr<IKeyExtension> KeyRenderer = CurrentViewModel.ImplicitCast();
-		if (KeyRenderer)
+		constexpr bool bIncludeThis = true;
+		for (FParentFirstChildIterator ChildIt(ViewModel, bIncludeThis) ; ChildIt; ++ChildIt)
 		{
-			KeyDrawInfo.Emplace(KeyRenderer);
-		}
+			FViewModelPtr CurrentViewModel = *ChildIt;
 
-		if (Params.bCollapseChildren == false)
-		{
-			break;
+			TViewModelPtr<FLinkedOutlinerExtension> OutlinerExtension = CurrentViewModel.ImplicitCast();
+			if (OutlinerExtension && OutlinerExtension->GetLinkedOutlinerItem() && OutlinerExtension->GetLinkedOutlinerItem()->IsFilteredOut())
+			{
+				ChildIt.IgnoreCurrentChildren();
+				continue;
+			}
+
+			TViewModelPtr<IKeyExtension> KeyRenderer = CurrentViewModel.ImplicitCast();
+			if (KeyRenderer)
+			{
+				KeyDrawInfo.Emplace(KeyRenderer);
+			}
+
+			if (Params.bCollapseChildren == false)
+			{
+				break;
+			}
 		}
 	}
 }
@@ -335,6 +422,8 @@ void FKeyRenderer::UpdateViewDependentData(const FKeyBatchParameters& Params, co
 
 		// Track whether all of the keys are within the valid range
 		bool bIsInRange = true;
+		// Track whether a key is on a whole frame or a sub frame
+		bool bIsOnSubFrame = false;
 
 		const FFrameTime ValidPlayRangeMin = Params.ValidPlayRangeMin;
 		const FFrameTime ValidPlayRangeMax = Params.ValidPlayRangeMax;
@@ -346,13 +435,17 @@ void FKeyRenderer::UpdateViewDependentData(const FKeyBatchParameters& Params, co
 		FFrameTime KeyTickEnd   = TNumericLimits<FFrameNumber>::Lowest();
 		EKeyConnectionStyle ConnectionStyle = EKeyConnectionStyle::None;
 
-		auto HandleKey = [&bIsInRange, &AverageKeyTime, &NumKeyTimes, &KeyTickStart, &KeyTickEnd, ValidPlayRangeMin, ValidPlayRangeMax](const FFrameTime& KeyFrame)
+		auto HandleKey = [&bIsInRange, &bIsOnSubFrame, &AverageKeyTime, &NumKeyTimes, &KeyTickStart, &KeyTickEnd
+			, ValidPlayRangeMin, ValidPlayRangeMax]
+		(const FFrameTime& KeyFrame)
 		{
 			if (bIsInRange && (KeyFrame < ValidPlayRangeMin || KeyFrame >= ValidPlayRangeMax))
 			{
 				bIsInRange = false;
 			}
 
+			bIsOnSubFrame = KeyFrame.GetSubFrame() != 0.f;
+
 			KeyTickStart = FMath::Min(KeyFrame, KeyTickStart);
 			KeyTickEnd   = FMath::Max(KeyFrame, KeyTickEnd);
 
@@ -487,6 +580,11 @@ void FKeyRenderer::UpdateViewDependentData(const FKeyBatchParameters& Params, co
 			NewKey.Flags |= EKeyRenderingFlags::OutOfRange;
 		}
 
+		if (bIsOnSubFrame)
+		{
+			NewKey.Flags |= EKeyRenderingFlags::SubFrame;
+		}
+
 		const int32 NumStyles = FMath::CountBits(static_cast<__underlying_type(EKeyConnectionStyle)>(NewKey.Params.ConnectionStyle));
 		if (NumStyles > 1)
 		{
@@ -656,6 +754,9 @@ int32 FKeyRenderer::DrawCurve(const FKeyBatchParameters& Params, const FGeometry
 
 int32 FKeyRenderer::Draw(const FKeyBatchParameters& Params, const FGeometry& AllottedGeometry, const FSlateRect& MyCullingRect, const FKeyRendererPaintArgs& PaintArgs, int32 LayerId) const
 {
+	const FVector2f LocalSize = AllottedGeometry.GetLocalSize();
+	const FVector2f HalfLocalSize = LocalSize * .5f;
+
 	// Draw key bars first
 	{
 		static const FSlateBrush* DottedKeyBarBrush = FAppStyle::GetBrush("Sequencer.KeyBar.Dotted");
@@ -666,7 +767,7 @@ int32 FKeyRenderer::Draw(const FKeyBatchParameters& Params, const FGeometry& All
 		const float KeyBarFadeThresholdPx = KeyBarFadeWidthPx * 0.25f;
 
 		static float KeyBarHeight = 2.f;
-		const float KeyBarTop = AllottedGeometry.GetLocalSize().Y * .5f - KeyBarHeight*.5f;
+		const float KeyBarTop = HalfLocalSize.Y - KeyBarHeight*.5f;
 		for (const FCachedKeyBar& KeyBar : PrecomputedKeyBars)
 		{
 			const float StartPx = Params.TimeToPixel.FrameToPixel(KeyBar.StartTime);
@@ -725,102 +826,168 @@ int32 FKeyRenderer::Draw(const FKeyBatchParameters& Params, const FGeometry& All
 		++LayerId;
 	}
 
+	static constexpr FLinearColor KeyBorderTint = FLinearColor(0.05f, 0.05f, 0.05f, 1.0f);
+	static constexpr FLinearColor KeyOverlapTint = FLinearColor(0.83f, 0.12f, 0.12f, 1.0f); // Red
+
 	for (const FKey& Key : PrecomputedKeys)
 	{
 		FKeyDrawParams KeyDrawParams = Key.Params;
 
-		if (EnumHasAnyFlags(Key.Flags, EKeyRenderingFlags::PartialKey))
+		KeyDrawParams.FillTint = PaintArgs.KeyFillColor;
+		KeyDrawParams.BorderTint = KeyBorderTint;
+
+		const bool bPartialKey = EnumHasAnyFlags(Key.Flags, EKeyRenderingFlags::PartialKey);
+		const bool bSelected = EnumHasAnyFlags(Key.Flags, EKeyRenderingFlags::Selected);
+		const bool bPartiallySelected = EnumHasAnyFlags(Key.Flags, EKeyRenderingFlags::AnySelected);
+		const bool bHovered = EnumHasAnyFlags(Key.Flags, EKeyRenderingFlags::Hovered);
+		const bool bSubFrame = EnumHasAnyFlags(Key.Flags, EKeyRenderingFlags::SubFrame);
+		const bool bOverlapped = EnumHasAnyFlags(Key.Flags, EKeyRenderingFlags::Overlaps);
+
+		if (bPartialKey)
 		{
 			KeyDrawParams.FillOffset = FVector2D(0.f, 0.f);
-			KeyDrawParams.FillTint   = KeyDrawParams.BorderTint  = FLinearColor::White;
-		}
+			KeyDrawParams.FillTint = PaintArgs.PartialKeyFillColor;
+			KeyDrawParams.BorderTint = PaintArgs.PartialKeyFillColor;
+		}		
 
-		const bool bSelected = EnumHasAnyFlags(Key.Flags, EKeyRenderingFlags::Selected);
 		// Determine the key color based on its selection/hover states
 		if (EnumHasAnyFlags(Key.Flags, EKeyRenderingFlags::PreviewSelected))
 		{
-			FLinearColor PreviewSelectionColor = PaintArgs.SelectionColor.LinearRGBToHSV();
+			FLinearColor PreviewSelectionColor = PaintArgs.KeySelectionFillColor.LinearRGBToHSV();
 			PreviewSelectionColor.R += 0.1f; // +10% hue
 			PreviewSelectionColor.G = 0.6f; // 60% saturation
 			KeyDrawParams.BorderTint = KeyDrawParams.FillTint = PreviewSelectionColor.HSVToLinearRGB();
 		}
 		else if (EnumHasAnyFlags(Key.Flags, EKeyRenderingFlags::PreviewNotSelected))
 		{
-			KeyDrawParams.BorderTint = FLinearColor(0.05f, 0.05f, 0.05f, 1.0f);
+			KeyDrawParams.BorderTint = KeyBorderTint;
 		}
 		else if (bSelected)
 		{
-			KeyDrawParams.BorderTint = PaintArgs.SelectionColor;
-			KeyDrawParams.FillTint = FLinearColor(0.05f, 0.05f, 0.05f, 1.0f);
-		}
-		else if (EnumHasAnyFlags(Key.Flags, EKeyRenderingFlags::AnySelected))
-		{
-			// partially selected
-			KeyDrawParams.BorderTint = PaintArgs.SelectionColor.CopyWithNewOpacity(0.5f);
-			KeyDrawParams.FillTint = FLinearColor(0.05f, 0.05f, 0.05f, 0.5f);
+			KeyDrawParams.BorderTint = PaintArgs.KeySelectionFillColor;
+			KeyDrawParams.FillTint = PaintArgs.KeySelectionFillColor;
 		}
-		else if (EnumHasAnyFlags(Key.Flags, EKeyRenderingFlags::Hovered))
+		else if (bPartiallySelected)
 		{
-			KeyDrawParams.BorderTint = FLinearColor(1.0f, 1.0f, 1.0f, 1.0f);
-			KeyDrawParams.FillTint = FLinearColor(1.0f, 1.0f, 1.0f, 1.0f);
+			KeyDrawParams.BorderTint = PaintArgs.PartialKeySelectionFillColor;
+			KeyDrawParams.FillTint =  PaintArgs.PartialKeySelectionFillColor;
 		}
-		else
+		else if (bHovered)
 		{
-			KeyDrawParams.BorderTint = FLinearColor(0.05f, 0.05f, 0.05f, 1.0f);
+			KeyDrawParams.BorderTint = PaintArgs.HoveredKeyFillColor;
+			KeyDrawParams.FillTint = PaintArgs.HoveredKeyFillColor;
 		}
 
 		// Color keys with overlaps with a red border
-		if (EnumHasAnyFlags(Key.Flags, EKeyRenderingFlags::Overlaps))
+		if (bOverlapped)
 		{
-			KeyDrawParams.BorderTint = FLinearColor(0.83f, 0.12f, 0.12f, 1.0f); // Red
+			KeyDrawParams.BorderTint = KeyOverlapTint;
 		}
 
-		const ESlateDrawEffect KeyDrawEffects = EnumHasAnyFlags(Key.Flags, EKeyRenderingFlags::OutOfRange) ? ESlateDrawEffect::DisabledEffect : PaintArgs.DrawEffects;
+		const ESlateDrawEffect KeyDrawEffects = EnumHasAnyFlags(Key.Flags, EKeyRenderingFlags::OutOfRange)
+			? ESlateDrawEffect::DisabledEffect : PaintArgs.DrawEffects;
 
 		// draw border
-		const FVector2D KeySize = bSelected ? Params.KeySizePx + PaintArgs.ThrobAmount * PaintArgs.KeyThrobValue : Params.KeySizePx;
-
-		static const float BrushBorderWidth = 2.0f;
-
 		const float KeyPositionPx = Params.TimeToPixel.FrameToPixel(Key.FinalKeyPosition);
 
-		FSlateDrawElement::MakeBox(
-			*PaintArgs.DrawElements,
-			// always draw selected keys on top of other keys
-			bSelected ? LayerId + 1 : LayerId,
-			// Center the key along Y.  Ensure the middle of the key is at the actual key time
-			AllottedGeometry.ToPaintGeometry(
+		FVector2D KeySize = bSelected
+			? Params.KeySizePx + (PaintArgs.ThrobAmount * PaintArgs.KeyThrobValue)
+			: Params.KeySizePx;
+		KeySize.X = FMath::Max(KeySize.X, 1.f);
+		const FVector2D HalfKeySize = KeySize * .5f;
+
+		if (PaintArgs.KeyStyle == EKeyRendererStyle::Line
+			|| PaintArgs.KeyStyle == EKeyRendererStyle::FrameBlock)
+		{
+			const FPaintGeometry KeyFrameGeometry = AllottedGeometry.ToPaintGeometry(
 				KeySize,
-				FSlateLayoutTransform(FVector2f(
-					KeyPositionPx - FMath::CeilToFloat(KeySize.X / 2.0f),
-					((AllottedGeometry.GetLocalSize().Y / 2.0f) - (KeySize.Y / 2.0f))
-				))
-			),
-			KeyDrawParams.BorderBrush,
-			KeyDrawEffects,
-			KeyDrawParams.BorderTint
-		);
+				FSlateLayoutTransform(FVector2D(KeyPositionPx, 0.) + PaintArgs.KeyOffset)
+			);
+
+			KeyDrawParams.FillBrush = bSubFrame ? FAppStyle::GetBrush(TEXT("Checker"))
+				: FAppStyle::GetBrush(TEXT("Brushes.White"));
+			KeyDrawParams.BorderBrush = FAppStyle::GetBrush(TEXT("PlainBorder"));
+
+			// Draw fill
+			FSlateDrawElement::MakeBox(
+				*PaintArgs.DrawElements,
+				bSelected ? LayerId + 1 : LayerId,
+				KeyFrameGeometry,
+				KeyDrawParams.FillBrush,
+				KeyDrawEffects,
+				KeyDrawParams.FillTint
+			);
+
+			// Draw border
+			if (PaintArgs.KeyStyle == EKeyRendererStyle::FrameBlock)
+			{
+				const bool bFrameBlockAggregateOverlap = PaintArgs.KeyStyle == EKeyRendererStyle::FrameBlock
+					&& Key.KeyTickStart != Key.KeyTickEnd;
+				const bool bShowFrameBlockBorder = bSelected
+					|| bPartiallySelected || bHovered || bOverlapped || bFrameBlockAggregateOverlap;
+				if (bShowFrameBlockBorder)
+				{
+					FSlateDrawElement::MakeBox(
+						*PaintArgs.DrawElements,
+						bSelected ? LayerId + 2 : LayerId + 1,
+						KeyFrameGeometry,
+						KeyDrawParams.BorderBrush,
+						KeyDrawEffects,
+						KeyDrawParams.BorderTint
+					);
+				}
+			}
+		}
+		else if (PaintArgs.KeyStyle == EKeyRendererStyle::Circle || PaintArgs.KeyStyle == EKeyRendererStyle::Diamond)
+		{
+			static constexpr float BrushBorderWidth = 2.f;
 
-		// draw fill
-		FSlateDrawElement::MakeBox(
-			*PaintArgs.DrawElements,
-			// always draw selected keys on top of other keys
-			bSelected ? LayerId + 2 : LayerId + 1,
-			// Center the key along Y.  Ensure the middle of the key is at the actual key time
-			AllottedGeometry.ToPaintGeometry(
-				KeySize - 2.0f * BrushBorderWidth,
-				FSlateLayoutTransform(
-					KeyDrawParams.FillOffset + 
-					FVector2D(
-						(KeyPositionPx - FMath::CeilToFloat((KeySize.X / 2.0f) - BrushBorderWidth)),
-						((AllottedGeometry.GetLocalSize().Y / 2.0f) - ((KeySize.Y / 2.0f) - BrushBorderWidth))
-					)
-				)
-			),
-			KeyDrawParams.FillBrush,
-			KeyDrawEffects,
-			KeyDrawParams.FillTint
-		);
+			if (PaintArgs.KeyStyle == EKeyRendererStyle::Diamond)
+			{
+				KeyDrawParams.FillBrush = FAppStyle::Get().GetBrush(TEXT("Sequencer.KeyDiamond"));
+			}
+			KeyDrawParams.ConnectionStyle = EKeyConnectionStyle::None;
+
+			// draw border
+			const FVector2D BorderOffset = FVector2D(
+				KeyPositionPx - FMath::CeilToFloat(HalfKeySize.X),
+				HalfLocalSize.Y - HalfKeySize.Y
+			);
+
+			FSlateDrawElement::MakeBox(
+				*PaintArgs.DrawElements,
+				// always draw selected keys on top of other keys
+				bSelected ? LayerId + 1 : LayerId,
+				// Center the key along Y.  Ensure the middle of the key is at the actual key time
+				AllottedGeometry.ToPaintGeometry(
+					KeySize,
+					FSlateLayoutTransform(BorderOffset + PaintArgs.KeyOffset)
+				),
+				KeyDrawParams.FillBrush,
+				KeyDrawEffects,
+				KeyDrawParams.BorderTint
+			);
+
+			// draw fill
+			const FVector2D FillOffset = FVector2D(
+				KeyPositionPx - FMath::CeilToFloat(HalfKeySize.X - BrushBorderWidth),
+				HalfLocalSize.Y - (HalfKeySize.Y - BrushBorderWidth)
+			);
+
+			FSlateDrawElement::MakeBox(
+				*PaintArgs.DrawElements,
+				// always draw selected keys on top of other keys
+				bSelected ? LayerId + 2 : LayerId + 1,
+				// Center the key along Y.  Ensure the middle of the key is at the actual key time
+				AllottedGeometry.ToPaintGeometry(
+					KeySize - (2.f * BrushBorderWidth),
+					FSlateLayoutTransform(FillOffset + KeyDrawParams.FillOffset + PaintArgs.KeyOffset)
+				),
+				KeyDrawParams.FillBrush,
+				KeyDrawEffects,
+				KeyDrawParams.FillTint
+			);
+		}
 	}
 
 	return LayerId + 2;
diff --git a/Engine/Source/Editor/SequencerCore/Private/MVVM/Views/STrackAreaView.cpp b/Engine/Source/Editor/SequencerCore/Private/MVVM/Views/STrackAreaView.cpp
--- a/Engine/Source/Editor/SequencerCore/Private/MVVM/Views/STrackAreaView.cpp
+++ b/Engine/Source/Editor/SequencerCore/Private/MVVM/Views/STrackAreaView.cpp
@@ -480,7 +480,10 @@ FReply STrackAreaView::OnMouseMove( const FGeometry& MyGeometry, const FPointerE
 	{
 		if (MouseEvent.IsMouseButtonDown(EKeys::RightMouseButton) && HasMouseCapture())
 		{
-			WeakOutliner.Pin()->ScrollByDelta(-MouseEvent.GetCursorDelta().Y);
+			if (const TSharedPtr<SOutlinerView> Outliner = WeakOutliner.Pin())
+			{
+				Outliner->ScrollByDelta(-MouseEvent.GetCursorDelta().Y);
+			}
 		}
 	}
 
@@ -505,8 +508,11 @@ FReply STrackAreaView::OnMouseWheel( const FGeometry& MyGeometry, const FPointer
 		return EditToolHandle;
 	}
 
-	const float ScrollAmount = -MouseEvent.GetWheelDelta() * GetGlobalScrollAmount();
-	WeakOutliner.Pin()->ScrollByDelta(ScrollAmount);
+	if (const TSharedPtr<SOutlinerView> Outliner = WeakOutliner.Pin())
+	{
+		const float ScrollAmount = -MouseEvent.GetWheelDelta() * GetGlobalScrollAmount();
+		Outliner->ScrollByDelta(ScrollAmount);
+	}
 
 	return FReply::Handled();
 }
@@ -519,11 +525,11 @@ void STrackAreaView::OnMouseEnter(const FGeometry& MyGeometry, const FPointerEve
 	DropFrameRange.Reset();
 
 	TSharedPtr<FTrackAreaViewModel> TrackArea = WeakViewModel.Pin();
-	ISequencerEditTool*             EditTool  = TrackArea ? TrackArea->GetEditTool() : nullptr;
+	ISequencerEditTool* const EditTool = TrackArea ? TrackArea->GetEditTool() : nullptr;
 
-	if (EditTool)
+	if (ISequencerInputHandler* const InputHandler = EditTool)
 	{
-		EditTool->OnMouseEnter(*this, MyGeometry, MouseEvent);
+		InputHandler->OnMouseEnter(*this, MyGeometry, MouseEvent);
 	}
 }
 
@@ -541,9 +547,9 @@ void STrackAreaView::OnMouseLeave(const FPointerEvent& MouseEvent)
 
 	if (TSharedPtr<FTrackAreaViewModel> TrackArea = WeakViewModel.Pin())
 	{
-		if (ISequencerEditTool* EditTool  = TrackArea->GetEditTool())
+		if (ISequencerInputHandler* const ToolInputHandler = TrackArea->GetEditTool())
 		{
-			EditTool->OnMouseLeave(*this, MouseEvent);
+			ToolInputHandler->OnMouseLeave(*this, MouseEvent);
 		}
 	}
 }
@@ -553,9 +559,9 @@ void STrackAreaView::OnMouseCaptureLost(const FCaptureLostEvent& CaptureLostEven
 {
 	if (TSharedPtr<FTrackAreaViewModel> TrackArea = WeakViewModel.Pin())
 	{
-		if (ISequencerEditTool* EditTool  = TrackArea->GetEditTool())
+		if (ISequencerInputHandler* const ToolInputHandler = TrackArea->GetEditTool())
 		{
-			EditTool->OnMouseCaptureLost();
+			ToolInputHandler->OnMouseCaptureLost();
 		}
 	}
 }
@@ -579,9 +585,9 @@ FCursorReply STrackAreaView::OnCursorQuery( const FGeometry& MyGeometry, const F
 			}
 		}
 
-		if (ISequencerEditTool* EditTool  = TrackArea->GetEditTool())
+		if (ISequencerInputHandler* const ToolInputHandler = TrackArea->GetEditTool())
 		{
-			return EditTool->OnCursorQuery(MyGeometry, CursorEvent);
+			return ToolInputHandler->OnCursorQuery(*this, MyGeometry, CursorEvent);
 		}
 	}
 
@@ -619,7 +625,7 @@ void STrackAreaView::OnDragEnter(const FGeometry& MyGeometry, const FDragDropEve
 {
 	SPanel::OnDragEnter(MyGeometry, DragDropEvent);
 }
-	
+
 void STrackAreaView::OnDragLeave(const FDragDropEvent& DragDropEvent)
 {
 	WeakDroppedItem = nullptr;
diff --git a/Engine/Source/Editor/SequencerCore/Public/ISequencerEditTool.h b/Engine/Source/Editor/SequencerCore/Public/ISequencerEditTool.h
--- a/Engine/Source/Editor/SequencerCore/Public/ISequencerEditTool.h
+++ b/Engine/Source/Editor/SequencerCore/Public/ISequencerEditTool.h
@@ -73,24 +73,42 @@ public:
 class ISequencerEditTool : public ISequencerInputHandler
 {
 public:
+	virtual ~ISequencerEditTool() override = default;
+
+	UE_DEPRECATED(5.8, "OnMouseCaptureLost has been deprecated. Please override OnMouseCaptureLost from ISequencerInputHandler instead.")
+	virtual void OnMouseCaptureLost() {}
+
+	virtual FCursorReply OnCursorQuery(const SWidget& OwnerWidget, const FGeometry& InGeometry, const FPointerEvent& InPointerEvent) const override
+	{
+		// Redirect to the old 2-argument version so existing overrides still work
+		PRAGMA_DISABLE_DEPRECATION_WARNINGS
+		return OnCursorQuery(InGeometry, InPointerEvent);
+		PRAGMA_ENABLE_DEPRECATION_WARNINGS
+	}
+
+	UE_DEPRECATED(5.8, "OnCursorQuery has been deprecated. Please override OnCursorQuery from ISequencerInputHandler instead.")
+	virtual FCursorReply OnCursorQuery(const FGeometry& InGeometry, const FPointerEvent& InPointerEvent) const { return FCursorReply::Unhandled(); }
+
+	UE_DEPRECATED(5.8, "OnMouseEnter has been deprecated. Please override OnMouseEnter from ISequencerInputHandler instead.")
+	virtual void OnMouseEnter(SWidget& OwnerWidget, const FGeometry& InGeometry, const FPointerEvent& InPointerEvent) {}
+
+	UE_DEPRECATED(5.8, "OnMouseLeave has been deprecated. Please override OnMouseLeave from ISequencerInputHandler instead.")
+	virtual void OnMouseLeave(SWidget& OwnerWidget, const FPointerEvent& InPointerEvent) {}
+
+	UE_DEPRECATED(5.8, "OnKeyDown has been deprecated. Please override OnKeyDown from ISequencerInputHandler instead.")
+	virtual FReply OnKeyDown(SWidget& OwnerWidget, const FGeometry& InGeometry, const FKeyEvent& InKeyEvent) { return FReply::Unhandled(); }
+
+	UE_DEPRECATED(5.8, "OnKeyUp has been deprecated. Please override OnKeyUp from ISequencerInputHandler instead.")
+	virtual FReply OnKeyUp(SWidget& OwnerWidget, const FGeometry& InGeometry, const FKeyEvent& InKeyEvent) { return FReply::Unhandled(); }
 
-	// @todo sequencer: documentation needed
-	virtual void OnMouseCaptureLost() = 0;
-	virtual int32 OnPaint(const FGeometry& AllottedGeometry, const FSlateRect& MyCullingRect, FSlateWindowElementList& OutDrawElements, int32 LayerId) const = 0;
-	virtual FCursorReply OnCursorQuery(const FGeometry& MyGeometry, const FPointerEvent& CursorEvent) const = 0;
-	virtual void OnMouseEnter(SWidget& OwnerWidget, const FGeometry& MyGeometry, const FPointerEvent& MouseEvent) = 0;
-	virtual void OnMouseLeave(SWidget& OwnerWidget, const FPointerEvent& MouseEvent) = 0;
-	virtual FReply OnKeyDown(SWidget& OwnerWidget, const FGeometry& MyGeometry, const FKeyEvent& InKeyEvent) {return FReply::Unhandled();}
-	virtual FReply OnKeyUp(SWidget& OwnerWidget, const FGeometry& MyGeometry, const FKeyEvent& InKeyEvent) {return FReply::Unhandled();}
 	virtual FName GetIdentifier() const = 0;
-	virtual bool CanDeactivate() const = 0;
-	virtual TSharedPtr<ITrackAreaHotspot> GetDragHotspot() const = 0;
-	virtual ISequencerEditToolDragOperation* GetCurrentDragOperation() const { return nullptr; }
 
-public:
+	virtual bool CanDeactivate() const { return true; }
 
-	/** Virtual destructor. */
-	virtual ~ISequencerEditTool() = default;
+	virtual int32 OnPaint(const FGeometry& InGeometry, const FSlateRect& InCullingRect, FSlateWindowElementList& OutDrawElements, int32 LayerId) const = 0;
+
+	virtual TSharedPtr<ITrackAreaHotspot> GetDragHotspot() const { return nullptr; }
+	virtual ISequencerEditToolDragOperation* GetCurrentDragOperation() const { return nullptr; }
 };
 
 } // namespace Sequencer
diff --git a/Engine/Source/Editor/SequencerCore/Public/ISequencerInputHandler.h b/Engine/Source/Editor/SequencerCore/Public/ISequencerInputHandler.h
--- a/Engine/Source/Editor/SequencerCore/Public/ISequencerInputHandler.h
+++ b/Engine/Source/Editor/SequencerCore/Public/ISequencerInputHandler.h
@@ -12,11 +12,22 @@ class SWidget;
  */
 struct ISequencerInputHandler
 {
-	virtual FReply OnMouseButtonDown(SWidget& OwnerWidget, const FGeometry& MyGeometry, const FPointerEvent& MouseEvent) = 0;
-	virtual FReply OnMouseButtonUp(SWidget& OwnerWidget, const FGeometry& MyGeometry, const FPointerEvent& MouseEvent) = 0;
-	virtual FReply OnMouseButtonDoubleClick(SWidget& OwnerWidget, const FGeometry& MyGeometry, const FPointerEvent& MouseEvent) = 0;
-	virtual FReply OnMouseMove(SWidget& OwnerWidget, const FGeometry& MyGeometry, const FPointerEvent& MouseEvent) = 0;
-	virtual FReply OnMouseWheel(SWidget& OwnerWidget, const FGeometry& MyGeometry, const FPointerEvent& MouseEvent) = 0;
-	virtual FReply OnKeyDown(SWidget& OwnerWidget, const FGeometry& MyGeometry, const FKeyEvent& InKeyEvent) { return FReply::Unhandled(); }
-	virtual FReply OnKeyUp(SWidget& OwnerWidget, const FGeometry& MyGeometry, const FKeyEvent& InKeyEvent) { return FReply::Unhandled(); }
+	virtual ~ISequencerInputHandler() = default;
+
+	virtual FCursorReply OnCursorQuery(const SWidget& OwnerWidget, const FGeometry& InGeometry, const FPointerEvent& InPointerEvent) const { return FCursorReply::Unhandled(); }
+
+	virtual void OnMouseEnter(SWidget& OwnerWidget, const FGeometry& InGeometry, const FPointerEvent& InPointerEvent) {}
+	virtual void OnMouseLeave(SWidget& OwnerWidget, const FPointerEvent& InPointerEvent) {}
+	virtual FReply OnMouseMove(SWidget& OwnerWidget, const FGeometry& InGeometry, const FPointerEvent& InPointerEvent) { return FReply::Unhandled(); }
+	virtual void OnMouseCaptureLost() {}
+
+	virtual FReply OnMouseButtonDown(SWidget& OwnerWidget, const FGeometry& InGeometry, const FPointerEvent& InPointerEvent) { return FReply::Unhandled(); }
+	virtual FReply OnMouseButtonUp(SWidget& OwnerWidget, const FGeometry& InGeometry, const FPointerEvent& InPointerEvent) { return FReply::Unhandled(); }
+	virtual FReply OnMouseButtonDoubleClick(SWidget& OwnerWidget, const FGeometry& InGeometry, const FPointerEvent& InPointerEvent) { return FReply::Unhandled(); }
+	virtual FReply OnMouseWheel(SWidget& OwnerWidget, const FGeometry& InGeometry, const FPointerEvent& InPointerEvent) { return FReply::Unhandled(); }
+
+	virtual FReply OnKeyDown(SWidget& OwnerWidget, const FGeometry& InGeometry, const FKeyEvent& InKeyEvent) { return FReply::Unhandled(); }
+	virtual FReply OnKeyUp(SWidget& OwnerWidget, const FGeometry& InGeometry, const FKeyEvent& InKeyEvent) { return FReply::Unhandled(); }
+
+	virtual void OnFocusLost(SWidget& OwnerWidget, const FFocusEvent& InFocusEvent) {}
 };
diff --git a/Engine/Source/Editor/SequencerCore/Public/MVVM/Views/KeyRenderer.h b/Engine/Source/Editor/SequencerCore/Public/MVVM/Views/KeyRenderer.h
--- a/Engine/Source/Editor/SequencerCore/Public/MVVM/Views/KeyRenderer.h
+++ b/Engine/Source/Editor/SequencerCore/Public/MVVM/Views/KeyRenderer.h
@@ -32,6 +32,15 @@ enum class EKeySelectionPreviewState : uint8
 	NotSelected
 };
 
+/** The style to draw a key */
+enum class EKeyRendererStyle : uint8
+{
+	Circle,
+	Diamond,
+	Line,
+	FrameBlock
+};
+
 /**
  * Paint arguments required for painting keys on a sequencer track
  */
@@ -52,14 +61,32 @@ struct FKeyRendererPaintArgs
 	/** Default draw effect */
 	ESlateDrawEffect DrawEffects;
 
-	/** Selection color to use for selected keys */
-	FLinearColor SelectionColor;
-
 	/** Curve color to paint underlying curves with */
 	FLinearColor CurveColor = FLinearColor(0.8f, 0.8f, 0.f, 0.7);
 
 	/** Key bar color */
 	FLinearColor KeyBarColor = FColor(160, 160, 160);
+
+	/** Style used to draw each key */
+	EKeyRendererStyle KeyStyle = EKeyRendererStyle::Circle;
+
+	/** Key draw offset */
+	FVector2D KeyOffset = FVector2D::ZeroVector;
+
+	/** The fill color to use to display keys on the timeline */
+	FLinearColor KeyFillColor = FLinearColor(.83f, .12f, .12f, 1.f);
+
+	/** The fill color to use to display hovered keys on the timeline */
+	FLinearColor HoveredKeyFillColor = FLinearColor::White;
+
+	/** The fill color to use to display selected keys on the timeline */
+	FLinearColor KeySelectionFillColor = HoveredKeyFillColor;
+
+	/** The fill color to use to display partial keys on the timeline */
+	FLinearColor PartialKeyFillColor = FLinearColor::White;
+
+	/** The fill color to use to display selected partial keys on the timeline */
+	FLinearColor PartialKeySelectionFillColor = KeySelectionFillColor;
 };
 
 /** Flag enum signifying how the cache has changed since it was last generated */
@@ -74,7 +101,6 @@ enum class EViewDependentCacheFlags : uint8
 	Empty            = 1 << 4,  // The view range is empty and there is nothing to show
 
 	All = DataChanged | KeyStateChanged | ViewChanged | ViewZoomed,
-
 };
 
 /** Interface for defining selection states for keys */
@@ -152,6 +178,7 @@ struct FKeyRenderer
 		Hovered            = 1 << 5,  // Only if NumKeys == NumHovered
 		Overlaps           = 1 << 6,  // If NumKeys > 1
 		OutOfRange         = 1 << 7,  // If any of the keys fall outside of the valid range
+		SubFrame           = 1 << 8,  // If the key is on a sub frame or a whole frame
 	};
 
 	enum class EKeyBarRenderingFlags
@@ -184,6 +211,11 @@ struct FKeyRenderer
 	 */
 	UE_API void Initialize(const FViewModelPtr& InViewModel);
 
+	/**
+	 * Initialize this renderer with multiple generic models
+	 */
+	UE_API void Initialize(const TSet<FWeakViewModelPtr>& InViewModels);
+
 	/**
 	 * Update any invalidated data required for drawing keys
 	 */
@@ -202,7 +234,7 @@ struct FKeyRenderer
 	/**
 	 * Retrieve keys at the specified position
 	 */
-	UE_API void HitTestKeys(const FFrameTime& Time, TArray<FKeysForModel>& OutAllKeys) const;
+	UE_API void HitTestKeys(const FFrameTime& Time, TArray<FKeysForModel>& OutAllKeys, const bool bInLeftAlignCollision = false) const;
 
 	/**
 	 * Retrieve keys at the specified position
@@ -328,8 +360,8 @@ private:
 
 private:
 
-	/** The model that we are drawing keys for */
-	FWeakViewModelPtr WeakViewModel;
+	/** The models that we are drawing keys for */
+	TSet<FWeakViewModelPtr> WeakViewModels;
 
 	/** Array of cached draw info for each of the key areas that comprise this batch */
 	mutable TArray<FCachedKeyDrawInformation, TInlineAllocator<1>> KeyDrawInfo;
diff --git a/Engine/Source/Editor/SequencerWidgets/Private/SSequencerTimeSlider.cpp b/Engine/Source/Editor/SequencerWidgets/Private/SSequencerTimeSlider.cpp
--- a/Engine/Source/Editor/SequencerWidgets/Private/SSequencerTimeSlider.cpp
+++ b/Engine/Source/Editor/SequencerWidgets/Private/SSequencerTimeSlider.cpp
@@ -26,6 +26,16 @@ int32 SSequencerTimeSlider::OnPaint( const FPaintArgs& Args, const FGeometry& Al
 	return FMath::Max( NewLayer, SCompoundWidget::OnPaint(Args, AllottedGeometry, MyCullingRect, OutDrawElements, NewLayer, InWidgetStyle, ShouldBeEnabled( bParentEnabled ) ) );
 }
 
+void SSequencerTimeSlider::OnMouseEnter(const FGeometry& MyGeometry, const FPointerEvent& MouseEvent)
+{
+	TimeSliderController->OnMouseEnter( *this, MyGeometry, MouseEvent );
+}
+
+void SSequencerTimeSlider::OnMouseLeave(const FPointerEvent& MouseEvent)
+{
+	TimeSliderController->OnMouseLeave( *this, MouseEvent );
+}
+
 FReply SSequencerTimeSlider::OnMouseButtonDown( const FGeometry& MyGeometry, const FPointerEvent& MouseEvent )
 {
 	TimeSliderController->OnTimeSliderMouseButtonDown( *this, MyGeometry, MouseEvent );
@@ -59,7 +69,12 @@ FReply SSequencerTimeSlider::OnMouseWheel( const FGeometry& MyGeometry, const FP
 
 FCursorReply SSequencerTimeSlider::OnCursorQuery(const FGeometry& MyGeometry, const FPointerEvent& CursorEvent) const
 {
-	return TimeSliderController->OnCursorQuery( SharedThis(this), MyGeometry, CursorEvent );
+	return TimeSliderController->OnCursorQuery( *this, MyGeometry, CursorEvent );
+}
+
+void SSequencerTimeSlider::OnFocusLost(const FFocusEvent& InFocusEvent)
+{
+	TimeSliderController->OnFocusLost( *this, InFocusEvent );
 }
 
 #undef LOCTEXT_NAMESPACE
diff --git a/Engine/Source/Editor/SequencerWidgets/Private/SSequencerTimeSlider.h b/Engine/Source/Editor/SequencerWidgets/Private/SSequencerTimeSlider.h
--- a/Engine/Source/Editor/SequencerWidgets/Private/SSequencerTimeSlider.h
+++ b/Engine/Source/Editor/SequencerWidgets/Private/SSequencerTimeSlider.h
@@ -37,15 +37,20 @@ public:
 	void Construct( const FArguments& InArgs, TSharedRef<ITimeSliderController> InTimeSliderController );
 
 protected:
-	// SWidget interface
+	//~ Begin SWidget
 	virtual FVector2D ComputeDesiredSize(float) const override;
 	virtual int32 OnPaint( const FPaintArgs& Args, const FGeometry& AllottedGeometry, const FSlateRect& MyCullingRect, FSlateWindowElementList& OutDrawElements, int32 LayerId, const FWidgetStyle& InWidgetStyle, bool bParentEnabled ) const override;
+	virtual void OnMouseEnter( const FGeometry& MyGeometry, const FPointerEvent& MouseEvent ) override;
+	virtual void OnMouseLeave( const FPointerEvent& MouseEvent ) override;
 	virtual FReply OnMouseButtonDown( const FGeometry& MyGeometry, const FPointerEvent& MouseEvent ) override;
 	virtual FReply OnMouseButtonUp( const FGeometry& MyGeometry, const FPointerEvent& MouseEvent ) override;
 	virtual FReply OnMouseButtonDoubleClick(const FGeometry& InMyGeometry, const FPointerEvent& InMouseEvent) override;
 	virtual FReply OnMouseMove( const FGeometry& MyGeometry, const FPointerEvent& MouseEvent ) override;
 	virtual FReply OnMouseWheel( const FGeometry& MyGeometry, const FPointerEvent& MouseEvent ) override;
 	virtual FCursorReply OnCursorQuery(const FGeometry& MyGeometry, const FPointerEvent& CursorEvent) const override;
+	virtual void OnFocusLost(const FFocusEvent& InFocusEvent) override;
+	//~ End SWidget
+
 private:
 	TSharedPtr<ITimeSliderController> TimeSliderController;
 	bool bMirrorLabels;
diff --git a/Engine/Source/Editor/SequencerWidgets/Public/ITimeSlider.h b/Engine/Source/Editor/SequencerWidgets/Public/ITimeSlider.h
--- a/Engine/Source/Editor/SequencerWidgets/Public/ITimeSlider.h
+++ b/Engine/Source/Editor/SequencerWidgets/Public/ITimeSlider.h
@@ -65,16 +65,12 @@ enum ETimeSliderPlaybackStatus
 class ITimeSliderController : public ISequencerInputHandler
 {
 public:
-	virtual ~ITimeSliderController()= default;
+	virtual ~ITimeSliderController() override = default;
 
 	virtual int32 OnPaintTimeSlider( bool bMirrorLabels, const FGeometry& AllottedGeometry, const FSlateRect& MyCullingRect, FSlateWindowElementList& OutDrawElements, int32 LayerId, const FWidgetStyle& InWidgetStyle, bool bParentEnabled ) const = 0;
 
 	virtual int32 OnPaintViewArea( const FGeometry& AllottedGeometry, const FSlateRect& MyCullingRect, FSlateWindowElementList& OutDrawElements, int32 LayerId, bool bEnabled, const FPaintViewAreaArgs& Args ) const = 0;
 
-	virtual FReply OnMouseButtonDoubleClick( SWidget& WidgetOwner, const FGeometry& MyGeometry, const FPointerEvent& MouseEvent ) override { return FReply::Unhandled(); }
-
-	virtual FCursorReply OnCursorQuery( TSharedRef<const SWidget> WidgetOwner, const FGeometry& MyGeometry, const FPointerEvent& CursorEvent ) const { return FCursorReply::Unhandled(); }
-
 	virtual FReply OnTimeSliderMouseButtonDown(SWidget& OwnerWidget, const FGeometry& MyGeometry, const FPointerEvent& MouseEvent)
 	{
 		return OnMouseButtonDown(OwnerWidget, MyGeometry, MouseEvent);
