Sequencer: Fix invisible camera cut track.

The problem was that during the drag operation the STracklLane and its underlying IOutlinerExtension are regenerated (see FTrackModel::ForceUpdate, TrackModel.cpp:280).  STrackAreaView::OnArrangeChildren:198 filters lanes by IsPinned() != bShowPinnedNodes would skip rendering the lane because the child lane was not considered pinned.

#Synced-CL 52378315

diff --git a/Engine/Source/Editor/SequencerCore/Private/MVVM/Views/STrackLane.cpp b/Engine/Source/Editor/SequencerCore/Private/MVVM/Views/STrackLane.cpp
--- a/Engine/Source/Editor/SequencerCore/Private/MVVM/Views/STrackLane.cpp
+++ b/Engine/Source/Editor/SequencerCore/Private/MVVM/Views/STrackLane.cpp
@@ -128,7 +128,13 @@ bool STrackLane::IsPinned() const
 	TSharedPtr<IOutlinerExtension> Outliner = WeakOutlinerItem.Pin();
 	TSharedPtr<IPinnableExtension> Pinnable = WeakOutlinerItem.ImplicitPin();
 
-	return (Pinnable && Pinnable->IsPinned()) || (Outliner && Outliner->ShouldAnchorToTop());
+	if ((Pinnable && Pinnable->IsPinned()) || (Outliner && Outliner->ShouldAnchorToTop()))
+	{
+		return true;
+	}
+
+	// When the parent lane is pinned, then the child lane should also act as if it was pinned.
+	return ParentLane && ParentLane->IsPinned();
 }
 
 void STrackLane::OnHierarchyUpdated()
